[NO NON-PII COMMIT MESSAGE]

diff --git a/Engine/Extras/RoboMerge/package.json b/Engine/Extras/RoboMerge/package.json
--- a/Engine/Extras/RoboMerge/package.json
+++ b/Engine/Extras/RoboMerge/package.json
@@ -14,7 +14,7 @@
   },
   "dependencies": {
     "@types/ldapjs": "^1.0.3",
-    "@types/node": "^10.3.2",
+    "@types/node": "^10.17.13",
     "@types/nodemailer": "^4.3.4",
     "jsonlint": "^1.6.3",
     "ldapauth-fork": "^4.0.2",
diff --git a/Engine/Extras/RoboMerge/src/common/webserver.ts b/Engine/Extras/RoboMerge/src/common/webserver.ts
--- a/Engine/Extras/RoboMerge/src/common/webserver.ts
+++ b/Engine/Extras/RoboMerge/src/common/webserver.ts
@@ -145,7 +145,7 @@ export class WebServer {
 	async close() {
 		if (this.server) {
 			// close the server to new connections
-			await new Promise<void>((done, _fail) => this.server!.close(done))
+			await new Promise((done, _fail) => this.server!.close(done))
 			this.server = null
 		}
 	}
