[NO NON-PII COMMIT MESSAGE]

diff --git a/Engine/Extras/RoboMerge/documentation/example-branchmap.json b/Engine/Extras/RoboMerge/documentation/example-branchmap.json
new file mode 100644
--- /dev/null
+++ b/Engine/Extras/RoboMerge/documentation/example-branchmap.json
@@ -0,0 +1,30 @@
+{
+     "defaultStreamDepot": "depot",
+     "isDefaultBot": true,
+     "branches": [
+         {
+             "name": "TestA",
+             "aliases": ["A"],
+             "flowsTo": ["TestB", "TestC"]
+         },
+         {
+             "name": "TestB",
+             "aliases": ["B"],
+             "flowsTo": ["TestA"],
+             "forceAll": true
+         },
+         {
+             "name": "TestC",
+             "aliases": ["C"],
+             "flowsTo": ["TestA"],
+             "forceAll": true
+         }
+     ],
+     "branchspecs": [
+         {
+             "name": "ROBO:TestA->TestC",
+             "from": "TestA",
+             "to": "TestC"
+         }
+     ]
+ }
\ No newline at end of file
