[NO NON-PII COMMIT MESSAGE]

diff --git a/Engine/Extras/RoboMerge/v3/src/robo/badges.ts b/Engine/Extras/RoboMerge/v3/src/robo/badges.ts
--- a/Engine/Extras/RoboMerge/v3/src/robo/badges.ts
+++ b/Engine/Extras/RoboMerge/v3/src/robo/badges.ts
@@ -5,7 +5,7 @@
 import { ContextualLogger } from '../common/logger';
 import { PerforceContext } from '../common/perforce';
 import { Blockage, Branch, BranchGraphInterface, ChangeInfo } from './branch-interfaces';
-import { BeginIntegratingToGateEvent, EndIntegratingToGateEvent } from './branch-interfaces';
+import { BeginIntegratingToGateEvent, GateEventContext } from './branch-interfaces';
 import { PersistentConflict, Resolution } from './conflict-interfaces';
 import { BotEventHandler, BotEvents } from './events';
 
@@ -97,13 +97,57 @@
 		}
 	}
 
+	/** On successful (possibly null) merge, mark whole chain (but not target branch) with same status */
+	private markIntegrationChain(info: ChangeInfo, result: string) {
+		const match = info.source.match(CHANGE_INFO_VIA_REGEX)
+		if (!match) {
+			return false
+		}
+
+		const badgeFunc = (cl: number, stream: string) => this.sendBadge(result, stream, cl, info.branch)
+
+		const sourceStream = match[1]
+		badgeFunc(info.source_cl, sourceStream)
+
+		// vias!
+		const allVias = match[2]
+		if (allVias) {
+			// note that at the point CL that commits to badge project is included in the vias
+			// (may not actually be committed anywhere else)
+			const viaClStrs = allVias.split(' via CL ')
+
+			// this will look up each changelist sequentially - no need to wait for it
+			// (maybe should technically wait for it at some level - could theoretically start marking a different status
+			// in parallel otherwise)
+			markViaBranches(viaClStrs.slice(1, viaClStrs.length - 1), badgeFunc, info.branch.parent, this.p4, this.badgeHandlerLogger)
+		}
+		return true
+	}
+
+	private badgeHandlerLogger : ContextualLogger
+	private botname: string
+	private externalUrl: string
+}
+
+class GateHandler implements BotEventHandler {
 	// for now just a map from target stream to timeout
 	// could in theory keep counters to support multiple sources
 	// should also wait for responses to deal with unpauses happening before pause is complete
 	private pausedFlows = new Map<string, NodeJS.Timeout>();
+	private logger: ContextualLogger
+
+	constructor(parentLogger: ContextualLogger) {
+		this.logger = parentLogger.createChild('GateHandler')
+	}
+
+	private static logStringFor(context: GateEventContext) {
+		return `${context.from.name} -> ${context.to.name} (@${context.edgeLastCl})`
+	}
+
 	onBeginIntegratingToGate(arg: BeginIntegratingToGateEvent) {
-		if (!this.pausedFlows.has(arg.to.upperName)) {
-			this.pausedFlows.set(arg.to.upperName, setTimeout(() => {
+		this.logger.info(GateHandler.logStringFor(arg.context) + ' catching up to CL# ' + arg.info.cl)
+		if (arg.context.pauseCIS && !this.pausedFlows.has(arg.context.to.upperName)) {
+			this.pausedFlows.set(arg.context.to.upperName, setTimeout(() => {
 				// @todo fill this in to pause!
 				const url = ''
 				const body = ''
@@ -111,13 +155,14 @@
 					url,
 					body,
 					contentType: 'application/json'
-				}, `Pause of CIS for ${arg.to}`)
+				}, `Pause of CIS for ${arg.context.to}`)
 
 			}, CIS_PAUSE_REFRESH_INTERVAL_SECONDS * 1000.))
 		}
 	}
 
-	onEndIntegratingToGate(arg: EndIntegratingToGateEvent) {
+	onEndIntegratingToGate(arg: GateEventContext) {
+		this.logger.info(GateHandler.logStringFor(arg) + ' caught up')
 		const pauseRefresher = this.pausedFlows.get(arg.to.upperName)
 		if (pauseRefresher) {
 			// @todo fill this in to unpause!
@@ -133,37 +178,6 @@
 			this.pausedFlows.delete(arg.to.upperName)
 		}
 	}
-
-	/** On successful (possibly null) merge, mark whole chain (but not target branch) with same status */
-	private markIntegrationChain(info: ChangeInfo, result: string) {
-		const match = info.source.match(CHANGE_INFO_VIA_REGEX)
-		if (!match) {
-			return false
-		}
-
-		const badgeFunc = (cl: number, stream: string) => this.sendBadge(result, stream, cl, info.branch)
-
-		const sourceStream = match[1]
-		badgeFunc(info.source_cl, sourceStream)
-
-		// vias!
-		const allVias = match[2]
-		if (allVias) {
-			// note that at the point CL that commits to badge project is included in the vias
-			// (may not actually be committed anywhere else)
-			const viaClStrs = allVias.split(' via CL ')
-
-			// this will look up each changelist sequentially - no need to wait for it
-			// (maybe should technically wait for it at some level - could theoretically start marking a different status
-			// in parallel otherwise)
-			markViaBranches(viaClStrs.slice(1, viaClStrs.length - 1), badgeFunc, info.branch.parent, this.p4, this.badgeHandlerLogger)
-		}
-		return true
-	}
-
-	private badgeHandlerLogger : ContextualLogger
-	private botname: string
-	private externalUrl: string
 }
 
 /** Add badges to branches listed as 'via' */
@@ -229,4 +243,6 @@
 		}
 		events.registerHandler(new BadgeHandler(badgeBranch, allStreams, externalUrl, logger))
 	}
+
+	events.registerHandler(new GateHandler(logger))
 }
diff --git a/Engine/Extras/RoboMerge/v3/src/robo/branch-interfaces.ts b/Engine/Extras/RoboMerge/v3/src/robo/branch-interfaces.ts
--- a/Engine/Extras/RoboMerge/v3/src/robo/branch-interfaces.ts
+++ b/Engine/Extras/RoboMerge/v3/src/robo/branch-interfaces.ts
@@ -196,13 +196,15 @@ export type GateInfo = {
 	date?: Date
 }
 
-export type BeginIntegratingToGateEvent = {
+export type GateEventContext = {
 	from: Branch
 	to: Branch
+	edgeLastCl: number
+	pauseCIS: boolean
+}
+
+export type BeginIntegratingToGateEvent = {
+	context: GateEventContext
 	info: GateInfo
 }
 
-export type EndIntegratingToGateEvent = {
-	from: Branch
-	to: Branch
-}
\ No newline at end of file
diff --git a/Engine/Extras/RoboMerge/v3/src/robo/edgebot.ts b/Engine/Extras/RoboMerge/v3/src/robo/edgebot.ts
--- a/Engine/Extras/RoboMerge/v3/src/robo/edgebot.ts
+++ b/Engine/Extras/RoboMerge/v3/src/robo/edgebot.ts
@@ -78,8 +78,12 @@ class EdgeBotImpl extends PerforceStatefulBot {
 		this.p4 = new PerforceContext(this.edgeBotLogger)
 		this.eventTriggers.onChangeParsed((info: ChangeInfo) => this.onGlobalChange(info))
 		this.gate = new Gate(
-			this.lastCl,
-			this.options.pauseCISUnlessAtGate ? {from: this.sourceBranch, to: this.branch, eventTriggers: this.eventTriggers} : null,
+			{ from: this.sourceBranch
+			, to: this.branch
+			, pauseCIS: !!this.options.pauseCISUnlessAtGate
+			, edgeLastCl: this.lastCl
+			, eventTriggers: this.eventTriggers
+			},
 			{options, p4: this.p4, logger: this.edgeBotLogger.createChild('gate')},
 			this.settings
 		)
diff --git a/Engine/Extras/RoboMerge/v3/src/robo/events.ts b/Engine/Extras/RoboMerge/v3/src/robo/events.ts
--- a/Engine/Extras/RoboMerge/v3/src/robo/events.ts
+++ b/Engine/Extras/RoboMerge/v3/src/robo/events.ts
@@ -1,7 +1,7 @@
 // Copyright Epic Games, Inc. All Rights Reserved.
 
 import { AlreadyIntegrated, Blockage, ChangeInfo, ForcedCl, PendingChange } from './branch-interfaces';
-import { BeginIntegratingToGateEvent, EndIntegratingToGateEvent } from './branch-interfaces';
+import { BeginIntegratingToGateEvent, GateEventContext } from './branch-interfaces';
 import { BotConfig } from './branchdefs';
 import { PersistentConflict } from './conflict-interfaces';
 
@@ -15,7 +15,7 @@ type OnNonSkipLastClChange = (arg: ForcedCl) => void
 type OnBranchUnblocked = (arg: PersistentConflict) => void
 type OnConflictStatus = (arg: boolean) => void
 type OnBeginIntegratingToGate = (arg: BeginIntegratingToGateEvent) => void
-type OnEndIntegratingToGate = (arg: EndIntegratingToGateEvent) => void
+type OnEndIntegratingToGate = (arg: GateEventContext) => void
 
 export interface BotEventHandler {
 	onChangeParsed?: OnChangeParsed
@@ -186,7 +186,7 @@ export class BotEventTriggers extends BotEvents {
 		}
 	}
 
-	reportEndIntegratingToGate(arg: EndIntegratingToGateEvent) {
+	reportEndIntegratingToGate(arg: GateEventContext) {
 		for (const listener of this.endIntegratingToGateListeners) {
 			listener(arg)
 		}
diff --git a/Engine/Extras/RoboMerge/v3/src/robo/gate.ts b/Engine/Extras/RoboMerge/v3/src/robo/gate.ts
--- a/Engine/Extras/RoboMerge/v3/src/robo/gate.ts
+++ b/Engine/Extras/RoboMerge/v3/src/robo/gate.ts
@@ -2,7 +2,7 @@
 
 import { ContextualLogger } from "../common/logger";
 import { Change, PerforceContext } from "../common/perforce";
-import { Branch, GateInfo } from "./branch-interfaces"
+import { GateEventContext, GateInfo } from "./branch-interfaces"
 import { DAYS_OF_THE_WEEK, EdgeOptions } from "./branchdefs"
 import { BotEventTriggers } from "./events"
 import { Context } from "./settings"
@@ -16,10 +16,16 @@ export type GateContext = {
 	logger: ContextualLogger
 }
 
-
 class DummyEventTriggers {
-	to: any = ''
+
+	// pretend GateEventContext 
 	from: any = ''
+	to: any = ''
+	pauseCIS = false
+
+	constructor(public edgeLastCl: number) {
+	}
+
 	static Inner = class {
 		beginCalls = 0
 		endCalls = 0
@@ -34,6 +40,8 @@ class DummyEventTriggers {
 
 	get beginCalls() { return this.eventTriggers.beginCalls }
 	get endCalls() { return this.eventTriggers.endCalls }
+
+	// named to match EventTriggersAndStuff.context
 	eventTriggers = new DummyEventTriggers.Inner
 }
 
@@ -95,33 +103,24 @@ async function getRequestedGateCl(context: GateContext): Promise<GateInfo | null
 	return result
 }
 
-type BranchesAndEventTriggers = {
-	from: Branch
-	to: Branch
+type EventTriggersAndStuff = GateEventContext & {
 	eventTriggers: BotEventTriggers
 }
 
 export class Gate {
-	constructor(private lastCl: number,
-					private eventTriggers: BranchesAndEventTriggers | DummyEventTriggers | null,
+	constructor(	private eventContext: EventTriggersAndStuff | DummyEventTriggers,
 					private context: GateContext, private persistence?: Context) {
 		this.loadFromPersistence()
 	}
 
-
 	private processGateChange(newGate: GateInfo) {
-/**
-Getting some complicated logic here - break down cases that are covered
-
-degrees of freedom:
-	order of lastCl/gate.cl/newGate.cl (L/G/N)
-	empty queue (!Q)
-
-	first block is N <= G || (G <= L && !Q)
+		// degrees of freedom:
+		// 	order of lastCl/gate.cl/newGate.cl (L/G/N)
+		// 	empty queue (!Q)
 
-	so remainder if N > G && G > L
- */
+		// 	first block is N <= G || (G <= L && !Q)
 
+		// 	so remainder if N > G && G > L
 
 		if (this.currentGateInfo) {
 			// check for special case of new gate being before (or equal to) current catch-up gate
@@ -208,10 +207,7 @@ degrees of freedom:
 			throw new Error('seem to have ignored gate!')
 		}
 
-/**
-
- */
-// if no current and in window, kick off new
+		// if no current and in window, kick off new
 		if (!this.currentGateInfo && this.queuedGates.length > 0) {
 			if (this.tryNextGate(this.lastCl)) {
 				this.reportCatchingUp()
@@ -225,32 +221,41 @@ degrees of freedom:
 	}
 
 	preIntegrate(cl: number) {
+		if (this.currentGateInfo) {
+			this.context.logger.verbose(`${this.lastCl} -> ${cl} ${this.currentGateInfo.cl}`)
+		}
+
+		let adjustedCl: number | null = null
 		if (this.currentGateInfo && cl > this.currentGateInfo.cl && this.currentGateInfo.cl > this.lastCl) {
 			// we're waiting for currentGateInfo.cl but got a higher cl
 			// basically we drag lastCl forward to currentGateInfo.cl so isGateOpen returns false
 			// but if there are queued gates, we should have a go at moving to the next gate instead
 
 			// note tryNextGate can change currentGateInfo or set it to null
-			if (!this.tryNextGate(cl) && this.currentGateInfo) {
-				this.lastCl = this.currentGateInfo.cl
-				return this.lastCl
+			if (!this.tryNextGate(cl)) {
+				this.reportCaughtUp()
+				if (this.currentGateInfo) {
+					adjustedCl = this.lastCl = this.currentGateInfo.cl
+				}
 			}
 		}
-		return null
+		return adjustedCl
 	}
 
 	/**
 	 Assumptions: currentGateInfo is valid and incoming cl > currentGateInfo.cl
 	 @return whether we started catching up to a new gate
-	 	- todo is this bool enough for the preIntegrate case, go though each case
 	 */
 	private tryNextGate(cl: number) {
 		if (this.queuedGates.length === 0) {
+			this.context.logger.verbose('nothing queued')
 			return false
 		}
 
 		const nextGateIndex = this.queuedGates.findIndex(queued => queued.cl > cl)
 		if (nextGateIndex < 0) {
+			this.context.logger.verbose('all done')
+
 			// all done (incoming CL is after all queued gates)
 			this.currentGateInfo = this.queuedGates[this.queuedGates.length - 1]
 			this.queuedGates.length = 0
@@ -258,11 +263,16 @@ degrees of freedom:
 		}
 
 		if (!this.nowIsWithinAllowedCatchupWindow()) {
+			// want to make this an info log, but would spam every tick at the moment
+			this.context.logger.verbose('delaying gate catch-up due to configured window')
+
 			// wait for next window before catching up with queued gates
 			this.currentGateInfo = null
 			return false
 		}
 
+		this.context.logger.verbose('next!')
+
 		// on to next gate
 		this.currentGateInfo = this.queuedGates[nextGateIndex]
 		this.queuedGates = this.queuedGates.slice(nextGateIndex + 1)
@@ -270,17 +280,21 @@ degrees of freedom:
 	}
 
 	setLastCl(cl: number) {
+		let notifyCaughtUp = false
 		if (this.currentGateInfo) {
 			// ignore going backwards; did we reach the gate?
 			if (cl > this.lastCl && cl >= this.currentGateInfo.cl) {
 				if (!this.tryNextGate(cl)) {
 					// either all done or waiting for next window to restart catching up
-					this.reportCaughtUp()
+					notifyCaughtUp = true
 				}
 			}
 		}
 
 		this.lastCl = cl
+		if (notifyCaughtUp) {
+			this.reportCaughtUp()
+		}
 	}
 
 	getNumChangesRemaining(changesFetched: Change[], changeIndex: number) {
@@ -311,7 +325,6 @@ degrees of freedom:
 		return !this.getGateClosedMessage()
 	}
 
-
 	getGateClosedMessage() {
 		// gate prevents integrations in two cases:
 		//	- we've caught up with the most recent gate
@@ -338,6 +351,13 @@ degrees of freedom:
 		}
 	}
 
+	logSummary() {
+		const logger = this.context.logger
+		logger.info('current gate: ' + (this.currentGateInfo ? this.currentGateInfo.cl.toString() : 'none'))
+		logger.info(`queued: ${this.queuedGates.length}`)
+		logger.info(`last cl: ${this.lastCl}`)
+	}
+
 	private nowIsWithinAllowedCatchupWindow() {
 		if (!this.context.options.integrationWindow) {
 			return true
@@ -377,22 +397,14 @@ degrees of freedom:
 		if (!mostRecent) {
 			throw new Error('what are we catching up to?')
 		}
-		if (this.eventTriggers) {
-			this.eventTriggers.eventTriggers.reportBeginIntegratingToGate({
-				to: this.eventTriggers.to,
-				from: this.eventTriggers.from,
-				info: mostRecent
-			})
-		}
+		this.eventContext.eventTriggers.reportBeginIntegratingToGate({
+			context: this.eventContext,
+			info: mostRecent
+		})
 	}
 
 	private reportCaughtUp() {
-		if (this.eventTriggers) {
-			this.eventTriggers.eventTriggers.reportEndIntegratingToGate({
-				to: this.eventTriggers.to,
-				from: this.eventTriggers.from
-			})
-		}
+		this.eventContext.eventTriggers.reportEndIntegratingToGate(this.eventContext)
 	}
 
 	private persist() {
@@ -426,11 +438,12 @@ degrees of freedom:
 		}
 	}
 
-	logSummary() {
-		const logger = this.context.logger
-		logger.info('current gate: ' + (this.currentGateInfo ? this.currentGateInfo.cl.toString() : 'none'))
-		logger.info(`queued: ${this.queuedGates.length}`)
-		logger.info(`last cl: ${this.lastCl}`)
+	private get lastCl() {
+		return this.eventContext.edgeLastCl
+	}
+
+	private set lastCl(cl: number) {
+		this.eventContext.edgeLastCl = cl
 	}
 
 	private currentGateInfo: GateInfo | null = null
@@ -447,8 +460,8 @@ export async function runTests(parentLogger: ContextualLogger) {
 	const logger = parentLogger.createChild('Gate')
 
 	const makeTestGate = (cl: number, options?: EdgeOptions): [DummyEventTriggers, Gate] => {
-		const et = new DummyEventTriggers()
-		return [et, new Gate(cl, et, { options: options || {}, p4: null, logger})]
+		const et = new DummyEventTriggers(cl)
+		return [et, new Gate(et, { options: options || {}, p4: null, logger})]
 	}
 
 	let fails = 0
