[NO NON-PII COMMIT MESSAGE]

diff --git a/Engine/Extras/RoboMerge/v3/run-tests.sh b/Engine/Extras/RoboMerge/v3/run-tests.sh
--- a/Engine/Extras/RoboMerge/v3/run-tests.sh
+++ b/Engine/Extras/RoboMerge/v3/run-tests.sh
@@ -1,3 +1,5 @@
+set -e
+
 RED='\033[0;31m'
 GREEN='\033[0;32m'
 NO_COLOUR='\033[0m'
@@ -32,29 +34,16 @@ cleanup() {
 
 cleanup
 
-docker network create --driver bridge robomerge_functtest_network &
+docker network create --driver bridge robomerge_functtest_network
 
 docker build -t helix -f helix.Dockerfile .
-docker build -t helix-and-node -f helix-and-node§.Dockerfile .
-
-docker build -t p4docker -f Dockerfile.p4docker . &
-docker build -t robomerge -f Dockerfile . &
-robo_build_pid=$!
+docker build -t helix-and-node -f helix-and-node.Dockerfile .
 
-docker build -t robomerge_functionaltests -f functional_tests/tstests.Dockerfile . &
-tests_build_pid=$!
+docker build -t p4docker -f Dockerfile.p4docker .
+docker build -t robomerge -f Dockerfile .
 
-wait $robo_build_pid
-if [ $? -ne 0 ]; then
-    exit 1
-fi
-
-wait $tests_build_pid
-if [ $? -ne 0 ]; then
-    exit 1
-fi
+docker build -t robomerge_functionaltests -f functional_tests/tstests.Dockerfile .
 
-wait
 
 echo "Running unit tests"
 docker run -a stderr -h robomerge_unittests --name robomerge_unittests robomerge npm test
