[NO NON-PII COMMIT MESSAGE]

diff --git a/Engine/Extras/RoboMerge/v3/src/robo/gate.ts b/Engine/Extras/RoboMerge/v3/src/robo/gate.ts
--- a/Engine/Extras/RoboMerge/v3/src/robo/gate.ts
+++ b/Engine/Extras/RoboMerge/v3/src/robo/gate.ts
@@ -140,6 +140,7 @@ export class Gate {
 			// case where we're already caught up and just need to replace the gate
 			if (this.currentGateInfo.cl <= this.lastCl) {
 				if (this.queuedGates.length > 0) {
+					this.context.logger.error('Unexpected queue: ' + this.queuedGates.map((info: GateInfo) => info.cl).join(', '))
 					throw new Error('invalid gate state!')
 				}
 				this.currentGateInfo = null
@@ -428,11 +429,14 @@ export class Gate {
 		}
 		const saved = this.persistence.get(GATE_INFO_KEY)
 		if (saved) {
+			this.context.logger.info(`Restoring saved gate info: current ${saved.current}`)
+
 			if (saved.current) {
 				this.currentGateInfo = saved.current
 				this.reportCatchingUp()
 			}
 			if (saved.queued) {
+				this.context.logger.info('Queue: ' + saved.queue.map((info: GateInfo) => info.cl).join(', '))
 				this.queuedGates = saved.queued
 			}
 		}
