[NO NON-PII COMMIT MESSAGE]

diff --git a/Engine/Extras/RoboMerge/v3/src/robo/gate.ts b/Engine/Extras/RoboMerge/v3/src/robo/gate.ts
--- a/Engine/Extras/RoboMerge/v3/src/robo/gate.ts
+++ b/Engine/Extras/RoboMerge/v3/src/robo/gate.ts
@@ -436,7 +436,7 @@ export class Gate {
 				this.reportCatchingUp()
 			}
 			if (saved.queued) {
-				this.context.logger.info('Queue: ' + saved.queue.map((info: GateInfo) => info.cl).join(', '))
+				this.context.logger.info('Queue: ' + saved.queued.map((info: GateInfo) => info.cl).join(', '))
 				this.queuedGates = saved.queued
 			}
 		}
