[NO NON-PII COMMIT MESSAGE]

diff --git a/Engine/Extras/RoboMerge/v3/src/robo/gate.ts b/Engine/Extras/RoboMerge/v3/src/robo/gate.ts
--- a/Engine/Extras/RoboMerge/v3/src/robo/gate.ts
+++ b/Engine/Extras/RoboMerge/v3/src/robo/gate.ts
@@ -429,15 +429,16 @@ export class Gate {
 		}
 		const saved = this.persistence.get(GATE_INFO_KEY)
 		if (saved) {
-			this.context.logger.info(`Restoring saved gate info: current ${saved.current}`)
+			this.context.logger.info(`Restoring saved gate info: current ${saved.current && saved.current.cl}`)
 
 			if (saved.current) {
 				this.currentGateInfo = saved.current
 				this.reportCatchingUp()
 			}
 			if (saved.queued) {
-				this.context.logger.info('Queue: ' + saved.queued.map((info: GateInfo) => info.cl).join(', '))
-				this.queuedGates = saved.queued
+				this.context.logger.warn('Restoring of queue temporarily disabled!')
+				// this.context.logger.info('Queue: ' + saved.queued.map((info: GateInfo) => info.cl).join(', '))
+				// this.queuedGates = saved.queued
 			}
 		}
 	}
