[NO NON-PII COMMIT MESSAGE]

diff --git a/Engine/Extras/RoboMerge/v3/src/robo/gate.ts b/Engine/Extras/RoboMerge/v3/src/robo/gate.ts
--- a/Engine/Extras/RoboMerge/v3/src/robo/gate.ts
+++ b/Engine/Extras/RoboMerge/v3/src/robo/gate.ts
@@ -129,7 +129,6 @@ export class Gate {
 			if (newGate.cl <= this.currentGateInfo.cl) {
 				// report if we were catching up and the replacement gate means we're now caught up
 				if (this.lastCl < this.currentGateInfo.cl && this.lastCl >= newGate.cl) {
-
 					this.reportCaughtUp()
 				}
 				this.currentGateInfo = newGate
@@ -138,18 +137,14 @@ export class Gate {
 			}
 
 			// case where we're already caught up and just need to replace the gate
-			if (this.currentGateInfo.cl <= this.lastCl) {
-				if (this.queuedGates.length > 0) {
-					this.context.logger.error('Unexpected queue: ' + this.queuedGates.map((info: GateInfo) => info.cl).join(', '))
-					throw new Error('invalid gate state!')
-				}
+			if (this.currentGateInfo.cl <= this.lastCl && this.queuedGates.length === 0) {
 				this.currentGateInfo = null
 				this.queuedGates.push(newGate)
 				return
 			}
 		}
 
-		// now we know either waiting for window or new gate is after catch-up gate, update the queue
+		// now we know we're either waiting for window or new gate is after catch-up gate, so update the queue
 
 		// e.g.
 		//	- queued gates at CLs [5, 8, 15]
@@ -436,9 +431,8 @@ export class Gate {
 				this.reportCatchingUp()
 			}
 			if (saved.queued) {
-				this.context.logger.warn('Restoring of queue temporarily disabled!')
-				// this.context.logger.info('Queue: ' + saved.queued.map((info: GateInfo) => info.cl).join(', '))
-				// this.queuedGates = saved.queued
+				this.context.logger.info('Queue: ' + saved.queued.map((info: GateInfo) => info.cl).join(', '))
+				this.queuedGates = saved.queued
 			}
 		}
 	}
@@ -470,12 +464,14 @@ export async function runTests(parentLogger: ContextualLogger) {
 	}
 
 	let fails = 0
+	let assertions = 0
 	let testName = ''
 	const assert = (b: boolean, msg: string) => {
 		if (!b) {
 			logger.error(`"${testName}" failed: ${colors.error(msg)}`)
 			++fails
 		}
+		++assertions
 	}
 
 	// rules (maybe encapsulate in helper functions)
@@ -506,6 +502,15 @@ export async function runTests(parentLogger: ContextualLogger) {
 		assert(et.beginCalls === 1, 'catching up')
 		assert(gate.isGateOpen(), 'gate open')
 
+		gate.setLastCl(2)
+		await gate.tick()
+
+		if (!exact) {
+			const newCl = gate.preIntegrate(4) // only know we've caught up when higher cl comes in
+			assert(newCl === 3, 'last cl adjustment requested')
+			gate.setLastCl(3)
+		}
+		assert(et.beginCalls === 1 && et.endCalls === 1, 'caught up')
 	}
 
 	const openWindow = (opts: EdgeOptions) => {
@@ -573,6 +578,34 @@ export async function runTests(parentLogger: ContextualLogger) {
 		assert(!gate.isGateOpen() && gate.getGateClosedMessage()!.includes('CIS'), 'gate closed')
 	}
 
+	const queueNoWindow = async (middleIntegration: boolean) => {
+ 		// add gates at 2 and 3 while 1 still pending, integrate 1, 2 and 3 (should make 2 optional)
+
+		const options: EdgeOptions = { lastGoodCLPath: 2 }
+
+		// initial CL is 1, gate is 2
+		const [et, gate] = makeTestGate(1, options)
+		await gate.tick()
+		assert(et.beginCalls === 1, 'no events yet') // initial catching up message
+		assert(gate.isGateOpen(), 'gate initally open')
+
+		options.lastGoodCLPath = 3
+
+		await gate.tick()
+		assert(et.beginCalls + et.endCalls === 1, 'no more events')
+		assert(gate.isGateOpen(), 'still open after gate queued')
+
+		// start integrating
+		gate.preIntegrate(1); gate.setLastCl(1); await gate.tick()
+		if (middleIntegration) {
+			gate.preIntegrate(2); gate.setLastCl(2); await gate.tick()
+		}
+		gate.preIntegrate(3); gate.setLastCl(3); await gate.tick()
+
+		assert(et.beginCalls === 1 && et.endCalls === 1, 'catch-up notified')
+
+	}
+
 	////
 	testName = 'golden path'
 	const [et1, gate1] = makeTestGate(1)
@@ -667,10 +700,18 @@ export async function runTests(parentLogger: ContextualLogger) {
 	testName = 'replace queued gates (before lastCl)'
 	await replaceQueuedGates('before lastCl')
 
-// lastGoodCLPath
-// integrationWindow
+	///
+	testName = 'queue no window'
+ 	await queueNoWindow(true)
+
+	///
+	testName = 'queue no window (skip integration)'
+ 	await queueNoWindow(false)
+
+	// try to test the case where I was missing a caught up message
+
 	if (fails === 0) {
-		logger.info(colors.info('Gate tests succeeded'))
+		logger.info(colors.info(`Gate tests succeeded (${assertions} assertions)`))
 	}
 	return fails
 }
