[NO NON-PII COMMIT MESSAGE]

diff --git a/Engine/Extras/RoboMerge/v3/functional_tests/framework/src/BranchState.ts b/Engine/Extras/RoboMerge/v3/functional_tests/framework/src/BranchState.ts
--- a/Engine/Extras/RoboMerge/v3/functional_tests/framework/src/BranchState.ts
+++ b/Engine/Extras/RoboMerge/v3/functional_tests/framework/src/BranchState.ts
@@ -26,6 +26,10 @@ export class EdgeState {
 		return this.rawEdgeData.lastGoodCL
 	}
 
+	getGateClosedMessage() {
+		return this.rawEdgeData.gateClosedMessage
+	}
+
 	dump() {
 		console.log(this.rawEdgeData)
 	}
diff --git a/Engine/Extras/RoboMerge/v3/functional_tests/framework/src/framework.ts b/Engine/Extras/RoboMerge/v3/functional_tests/framework/src/framework.ts
--- a/Engine/Extras/RoboMerge/v3/functional_tests/framework/src/framework.ts
+++ b/Engine/Extras/RoboMerge/v3/functional_tests/framework/src/framework.ts
@@ -163,8 +163,6 @@ type RobomergeBranchOptions = {
 	incognitoMode: boolean
 
 	excludeAuthors: string[] // if present, completely overrides BotConfig
-
-	p4MaxRowsOverride: number // use with care and check with p4 admins
 }
 
 export type RobomergeBranchSpec = Partial<RobomergeBranchOptions> & {
@@ -172,6 +170,13 @@ export type RobomergeBranchSpec = Partial<RobomergeBranchOptions> & {
 }
 
 // copied from branchdefs.ts
+type IntegrationWindowPane = {
+	// if day not specified, daily
+	dayOfTheWeek?: 'sun' | 'mon' | 'tue' | 'wed' | 'thu' | 'fri' | 'sat'
+	startHourUTC: number
+	durationHours: number
+}
+
 type EdgeOptionFields = {
 	lastGoodCLPath: string
 	additionalSlackChannel: string
@@ -183,6 +188,10 @@ type EdgeOptionFields = {
 	terminal: boolean // changes go along terminal edges but no further
 
 	excludeAuthors: string[]
+
+	// by default, specify when gate catch ups are allowed; can be inverted to disallow
+	integrationWindow: IntegrationWindowPane[]
+	invertIntegrationWindow: boolean
 }
 
 export type EdgeProperties = Partial<EdgeOptionFields> & {
@@ -810,10 +819,14 @@ export abstract class FunctionalTest {
 				const p4CL = await latestP4CLs.get(source)!
 				const rmCL = edgeState.getLastCL()
 				if (rmCL < p4CL) {
-					const gateCL = edgeState.getLastGoodCL()
-					if (!gateCL || gateCL < 0 || rmCL < gateCL) {
+					const gateClosed = edgeState.getGateClosedMessage()
+					if (gateClosed) {
+						this.verbose(gateClosed)
+					}
+					else {
 						if (dump) {
 							let msg = `RoboMerge is not idle: ${source} -> ${target} last CL ${rmCL} < ${p4CL}`
+							const gateCL = edgeState.getLastGoodCL()
 							if (gateCL) {
 								msg += ` (gate: ${gateCL})`
 							}
diff --git a/Engine/Extras/RoboMerge/v3/functional_tests/framework/src/tests/test-edge-gate.ts b/Engine/Extras/RoboMerge/v3/functional_tests/framework/src/tests/test-edge-gate.ts
--- a/Engine/Extras/RoboMerge/v3/functional_tests/framework/src/tests/test-edge-gate.ts
+++ b/Engine/Extras/RoboMerge/v3/functional_tests/framework/src/tests/test-edge-gate.ts
@@ -1,11 +1,13 @@
 // Copyright Epic Games, Inc. All Rights Reserved.
 import { FunctionalTest, getRootDataClient, P4Client, P4Util, Stream } from '../framework'
+import * as system from '../system'
 
 const streams: Stream[] = [
 	{name: 'Main', streamType: 'mainline'},
 	{name: 'Dev-NoGate', streamType: 'development', parent: 'Main'},
 	{name: 'Dev-PlusOne', streamType: 'development', parent: 'Main'},
-	{name: 'Dev-Exact', streamType: 'development', parent: 'Main'}
+	{name: 'Dev-Exact', streamType: 'development', parent: 'Main'},
+	{name: 'Dev-Queue', streamType: 'development', parent: 'Main'}
 ]
 
 const GATE_FILENAME = 'TestEdgeGate-gate'
@@ -26,41 +28,60 @@ export class TestEdgeGate extends FunctionalTest {
 		await Promise.all([
 			this.p4.populate(this.getStreamPath('Dev-Exact'), desc),
 			this.p4.populate(this.getStreamPath('Dev-PlusOne'), desc),
-			this.p4.populate(this.getStreamPath('Dev-NoGate'), desc)
+			this.p4.populate(this.getStreamPath('Dev-NoGate'), desc),
+			this.p4.populate(this.getStreamPath('Dev-Queue'), desc)
 		])
 
 		const firstEditCl = await P4Util.editFileAndSubmit(mainClient, 'test.txt', 'Initial content\n\nFirst addition')
 
 		// must happen in set-up, and for this test, before second addition (i.e. earlier CL numbers)
 		await Promise.all([
-			P4Util.addFile(this.gateClient, GATE_FILENAME + 'exact.json', `{"Change":${firstEditCl}}`),
-			P4Util.addFile(this.gateClient, GATE_FILENAME + 'plusone.json', `{"Change":${firstEditCl + 1}}`)
+			P4Util.addFile(this.gateClient, GATE_FILENAME + 'exact.json', JSON.stringify({Change:firstEditCl})),
+			P4Util.addFile(this.gateClient, GATE_FILENAME + 'plusone.json', JSON.stringify({Change:firstEditCl + 1})),
+			P4Util.addFile(this.gateClient, GATE_FILENAME + 'queue.json', JSON.stringify({Change:firstEditCl, integrationWindow: []}))
 		])
 
 		await this.gateClient.submit('Added gates')
-		await P4Util.editFileAndSubmit(mainClient, 'test.txt', 'Initial content\n\nFirst addition\n\nSecond addition')
+		const contentsWithTwoAdditions = 'Initial content\n\nFirst addition\n\nSecond addition'
+		this.secondEditCl = await P4Util.editFileAndSubmit(mainClient, 'test.txt', contentsWithTwoAdditions)
+		await P4Util.editFileAndSubmit(mainClient, 'test.txt', contentsWithTwoAdditions + '\n\nThird addition')
 	}
 
+	secondEditCl = -1
+
 	async run() {
+		await this.waitForRobomergeIdle()
+		await this.checkHeadRevision('Dev-Queue', 'test.txt', 1)
+		await P4Util.editFileAndSubmit(this.gateClient, GATE_FILENAME + 'queue.json', JSON.stringify({Change:this.secondEditCl, integrationWindow: []}))
+
+		this.info('updated gate but window closed; sleeping') // don't have a way to wait for robomerge to pick up the gate change
+		await system.sleep(10)
+		await this.checkHeadRevision('Dev-Queue', 'test.txt', 1)
+
+		await P4Util.editFileAndSubmit(this.gateClient, GATE_FILENAME + 'queue.json', JSON.stringify({Change:this.secondEditCl}))
+		// this.info('removing window, sleeping again')
+		// await system.sleep(10)
 	}
 
 	verify() {
 		return Promise.all([
-			this.checkHeadRevision('Main', 'test.txt', 3),
+			this.checkHeadRevision('Main', 'test.txt', 4),
 			this.checkHeadRevision('Dev-Exact', 'test.txt', 2),
 			this.checkHeadRevision('Dev-PlusOne', 'test.txt', 2),
-			this.checkHeadRevision('Dev-NoGate', 'test.txt', 3)
+			this.checkHeadRevision('Dev-Queue', 'test.txt', 3),
+			this.checkHeadRevision('Dev-NoGate', 'test.txt', 4)
 		])
 	}
 
 	getBranches() {
-		const mainSpec = this.makeForceAllBranchDef('Main', ['Dev-Exact', 'Dev-NoGate', 'Dev-PlusOne'])
+		const mainSpec = this.makeForceAllBranchDef('Main', ['Dev-Exact', 'Dev-NoGate', 'Dev-PlusOne', 'Dev-Queue'])
 		mainSpec.initialCL = 1
 		return [
 			mainSpec,
 			this.makeForceAllBranchDef('Dev-Exact', []),
 			this.makeForceAllBranchDef('Dev-PlusOne', []),
-			this.makeForceAllBranchDef('Dev-NoGate', [])
+			this.makeForceAllBranchDef('Dev-NoGate', []),
+			this.makeForceAllBranchDef('Dev-Queue', [])
 		]
 	}
 
@@ -72,6 +93,9 @@ export class TestEdgeGate extends FunctionalTest {
 		, { from: this.fullBranchName('Main'), to: this.fullBranchName('Dev-PlusOne')
 		  , lastGoodCLPath: this.gateClient.stream + '/' + GATE_FILENAME + 'plusone.json'
 		  }
+		, { from: this.fullBranchName('Main'), to: this.fullBranchName('Dev-Queue')
+		  , lastGoodCLPath: this.gateClient.stream + '/' + GATE_FILENAME + 'queue.json'
+		  }
 		]
 	}
 }
diff --git a/Engine/Extras/RoboMerge/v3/src/robo/branch-interfaces.ts b/Engine/Extras/RoboMerge/v3/src/robo/branch-interfaces.ts
--- a/Engine/Extras/RoboMerge/v3/src/robo/branch-interfaces.ts
+++ b/Engine/Extras/RoboMerge/v3/src/robo/branch-interfaces.ts
@@ -4,7 +4,7 @@ import { BranchSpec, ConflictedResolveNFile, RoboWorkspace } from '../common/per
 // TODO: Remove Circular Dependency on bot-interfaces
 import { NodeBotInterface } from './bot-interfaces';
 import { FailureKind } from './status-types';
-import { BotConfig, BranchBase, EdgeOptions, NodeOptions } from './branchdefs';
+import { BotConfig, BranchBase, EdgeOptions, NodeOptions, IntegrationWindowPane } from './branchdefs';
 import { BlockageNodeOpUrls } from './roboserver';
 
 export type BranchArg = Branch | string
@@ -194,6 +194,53 @@ export type GateInfo = {
 	cl: number
 	link?: string
 	date?: Date
+
+	// optional overrides for integration window (takes precedence over any in config)
+	integrationWindow?: IntegrationWindowPane[]
+	invertIntegrationWindow?: boolean
+}
+
+export function gatesSame(lhs: GateInfo | null, rhs: GateInfo | null) {
+	if (!lhs && !rhs) {
+		return true
+	}
+
+	if (!lhs || !rhs) {
+		return false
+	}
+
+	if (lhs.cl !== rhs.cl) {
+		return false
+	}
+
+	// neither has windows? no need to check further
+	if (!lhs.integrationWindow && !rhs.integrationWindow) {
+		return true
+	}
+
+	if (!lhs.integrationWindow || !rhs.integrationWindow) {
+		return false
+	}
+
+	if (lhs.integrationWindow.length !== rhs.integrationWindow.length) {
+		return false
+	}
+
+	if (!lhs.invertIntegrationWindow !== !rhs.invertIntegrationWindow) {
+		return false
+	}
+
+	for (let n = 0; n < lhs.integrationWindow.length; ++n) {
+		const lhsWindow = lhs.integrationWindow[n]
+		const rhsWindow = rhs.integrationWindow[n]
+		if (lhsWindow.startHourUTC !== rhsWindow.startHourUTC ||
+				lhsWindow.durationHours !== rhsWindow.durationHours ||
+				lhsWindow.dayOfTheWeek !== rhsWindow.dayOfTheWeek) {
+			return false
+		}
+	}
+
+	return true
 }
 
 export type GateEventContext = {
@@ -207,4 +254,3 @@ export type BeginIntegratingToGateEvent = {
 	context: GateEventContext
 	info: GateInfo
 }
-
diff --git a/Engine/Extras/RoboMerge/v3/src/robo/branchdefs.ts b/Engine/Extras/RoboMerge/v3/src/robo/branchdefs.ts
--- a/Engine/Extras/RoboMerge/v3/src/robo/branchdefs.ts
+++ b/Engine/Extras/RoboMerge/v3/src/robo/branchdefs.ts
@@ -60,7 +60,7 @@ export class IntegrationMethod {
 export const DAYS_OF_THE_WEEK = ['sun', 'mon', 'tue', 'wed', 'thu', 'fri', 'sat']
 type DayOfTheWeek = 'sun' | 'mon' | 'tue' | 'wed' | 'thu' | 'fri' | 'sat'
 
-type IntegrationWindowPane = {
+export type IntegrationWindowPane = {
 	// if day not specified, daily
 	dayOfTheWeek?: DayOfTheWeek
 	startHourUTC: number
diff --git a/Engine/Extras/RoboMerge/v3/src/robo/gate.ts b/Engine/Extras/RoboMerge/v3/src/robo/gate.ts
--- a/Engine/Extras/RoboMerge/v3/src/robo/gate.ts
+++ b/Engine/Extras/RoboMerge/v3/src/robo/gate.ts
@@ -2,8 +2,8 @@
 
 import { ContextualLogger } from "../common/logger";
 import { Change, PerforceContext } from "../common/perforce";
-import { GateEventContext, GateInfo } from "./branch-interfaces"
-import { DAYS_OF_THE_WEEK, EdgeOptions } from "./branchdefs"
+import { gatesSame, GateEventContext, GateInfo } from "./branch-interfaces"
+import { DAYS_OF_THE_WEEK, EdgeOptions, IntegrationWindowPane } from "./branchdefs"
 import { BotEventTriggers } from "./events"
 import { Context } from "./settings"
 
@@ -45,7 +45,7 @@ class DummyEventTriggers {
 	eventTriggers = new DummyEventTriggers.Inner
 }
 
-async function getRequestedGateCl(context: GateContext): Promise<GateInfo | null> {
+async function getRequestedGateCl(context: GateContext, previousGateInfo?: GateInfo | null): Promise<GateInfo | null> {
 
 	const lastGoodCLPath = context.options.lastGoodCLPath
 	if (!lastGoodCLPath) {
@@ -91,13 +91,28 @@ async function getRequestedGateCl(context: GateContext): Promise<GateInfo | null
 	if (clInfo.Url) {
 		result.link = clInfo.Url
 	}
-	try {
-		const description = await context.p4!.describe(cl)
-		if (description.date) {
-			result.date = description.date
+	if (previousGateInfo && previousGateInfo.cl === cl) {
+		if (previousGateInfo.date) {
+			result.date = previousGateInfo.date
 		}
 	}
-	catch (_err) {
+	else {
+		try {
+			const description = await context.p4!.describe(cl)
+			if (description.date) {
+				result.date = description.date
+			}
+		}
+		catch (err) {
+			console.log('error getting gate CL description', err)
+		}
+	}
+
+	if (clInfo.integrationWindow) {
+		result.integrationWindow = clInfo.integrationWindow
+		if (clInfo.invertIntegrationWindow) {
+			result.invertIntegrationWindow = true
+		}
 	}
 
 	return result
@@ -169,7 +184,9 @@ export class Gate {
 			return
 		}
 
-		const gateInfo = await getRequestedGateCl(this.context)
+		const mostRecentGate = this.getMostRecentGate()
+
+		const gateInfo = await getRequestedGateCl(this.context, mostRecentGate)
 		if (!gateInfo) {
 			if (this.currentGateInfo) {
 				// we were waiting for a gate, so unpause and clear gate info
@@ -186,14 +203,12 @@ export class Gate {
 		}
 
 		let dirty = false
-		const mostRecentGate = this.getMostRecentGate()
-
 		if (!mostRecentGate) {
 			// ooh look, our first gate
 			this.queuedGates.push(gateInfo)
 			dirty = true
 		}
-		else if (mostRecentGate.cl !== gateInfo.cl) {
+		else if (!gatesSame(mostRecentGate, gateInfo)) {
 			this.processGateChange(gateInfo)
 			dirty = true
 		}
@@ -345,6 +360,11 @@ export class Gate {
 			outStatus.lastGoodCLJobLink = mostRecentGate.link
 			outStatus.lastGoodCLDate = mostRecentGate.date
 		}
+
+		const closedMessage = this.getGateClosedMessage()
+		if (closedMessage) {
+			outStatus.gateClosedMessage = closedMessage
+		}
 	}
 
 	logSummary() {
@@ -355,14 +375,27 @@ export class Gate {
 	}
 
 	private nowIsWithinAllowedCatchupWindow() {
-		if (!this.context.options.integrationWindow) {
+		const mostRecentGate = this.getMostRecentGate()
+		let integrationWindow: IntegrationWindowPane[] | null = null
+		let invert = false
+
+		if (mostRecentGate && mostRecentGate.integrationWindow) {
+			integrationWindow = mostRecentGate.integrationWindow
+			invert = !!mostRecentGate.invertIntegrationWindow
+		}
+		else if (this.context.options.integrationWindow) {
+			integrationWindow = this.context.options.integrationWindow
+			invert = !!this.context.options.invertIntegrationWindow
+		}
+
+		if (!integrationWindow) {
 			return true
 		}
 
 		const now = new Date;
 		let inWindow = false
 
-		for (const pane of this.context.options.integrationWindow) {
+		for (const pane of integrationWindow) {
 			if (pane.dayOfTheWeek) {
 				const dayIndex = DAYS_OF_THE_WEEK.indexOf(pane.dayOfTheWeek)
 				if (dayIndex < 0 || dayIndex > 6) {
@@ -382,7 +415,7 @@ export class Gate {
 			}
 		}
 
-		if (this.context.options.invertIntegrationWindow) {
+		if (invert) {
 			inWindow = !inWindow
 		}
 		return inWindow
@@ -516,7 +549,7 @@ export async function runTests(parentLogger: ContextualLogger) {
 	const openWindow = (opts: EdgeOptions) => {
 		opts.integrationWindow!.push({
 			startHourUTC: (new Date).getUTCHours(),
-			durationHours: 1
+			durationHours: 2 // more than 1 to avoid edge cases
 			})
 	}
 
