[NO NON-PII COMMIT MESSAGE]

diff --git a/Engine/Extras/RoboMerge/v3/src/robo/targets.ts b/Engine/Extras/RoboMerge/v3/src/robo/targets.ts
--- a/Engine/Extras/RoboMerge/v3/src/robo/targets.ts
+++ b/Engine/Extras/RoboMerge/v3/src/robo/targets.ts
@@ -170,7 +170,8 @@ export class DescriptionParser {
 			// get impaled on our own commit hook
 			return
 		}
-		
+
+		// check for any non-robomerge # tags we care about
 		if (ws || !command.startsWith('ROBOMERGE')) {
 			// check for commands to neuter
 			if (NEUTER_COMMANDS.indexOf(command) >= 0) {
@@ -179,7 +180,7 @@ export class DescriptionParser {
 			}
 			
 			else if (command.indexOf('REVIEW-') === 0) {
-				// remove swarm review numbers entirely (The can't be neutered and generate emails)
+				// remove swarm review numbers entirely (they can't be neutered and generate emails)
 				if (value !== '') {
 					this.descFinal.push(value)
 				}
@@ -199,27 +200,10 @@ export class DescriptionParser {
 		// Handle commands
 		if (command === 'ROBOMERGE') {
 
-			const targets = parseTargetList(value)
 			// ignore bare ROBOMERGE tags if we're not the default bot (should not affect default flow)
 			if (this.isDefaultBot) {
 				this.useDefaultFlow = false
-				for (const target of targets) {
-					const macroLines = this.macros[target.toLowerCase()]
-					if (macroLines) {
-						this.expandedMacros.push(target)
-						this.expandedMacroLines = [...this.expandedMacroLines, ...macroLines]
-					}
-					else {
-						this.arguments.push(target)
-					}
-				}
-			}
-			else {
-				for (const target of targets) {
-					if (this.macros[target.toLowerCase()]) {
-						this.errors.push(`Macro (${target}) not allowed by this bot`)
-					}
-				}
+				this.processTargetList(value)
 			}
 			return
 		}
@@ -230,12 +214,12 @@ export class DescriptionParser {
 		if (specificBotMatch) {
 			const specificBot = specificBotMatch[1].toUpperCase()
 			if (specificBot === this.graphBotName ||
-				specificBot === this.aliasUpper || 
+				specificBot === this.aliasUpper ||
 				specificBot === 'ALL') {
 				this.useDefaultFlow = false
 
 // @todo 'ALL' with targets should be an error (but don't know at this point ...)
-				this.arguments = [...this.arguments, ...parseTargetList(value)]
+				this.processTargetList(value)
 			}
 			else {
 				// keep a record of commands to forward to other bots - processed later
@@ -269,6 +253,21 @@ export class DescriptionParser {
 			this.errors.push(`Unknown command '${command}`)
 		}
 	}
+
+	private processTargetList(targetString: string) {
+		const targetNames = parseTargetList(targetString)
+
+		for (const name of targetNames) {
+			const macroLines = this.macros[name.toLowerCase()]
+			if (macroLines) {
+				this.expandedMacros.push(name)
+				this.expandedMacroLines = [...this.expandedMacroLines, ...macroLines]
+			}
+			else {
+				this.arguments.push(name)
+			}
+		}
+	}
 }
 
 export function parseDescriptionLines(
