[NO NON-PII COMMIT MESSAGE]

diff --git a/Engine/Extras/RoboMerge/v3/functional_tests/framework/src/tests/test-edge-gate.ts b/Engine/Extras/RoboMerge/v3/functional_tests/framework/src/tests/test-edge-gate.ts
--- a/Engine/Extras/RoboMerge/v3/functional_tests/framework/src/tests/test-edge-gate.ts
+++ b/Engine/Extras/RoboMerge/v3/functional_tests/framework/src/tests/test-edge-gate.ts
@@ -59,11 +59,12 @@ export class TestEdgeGate extends FunctionalTest {
 		await this.checkHeadRevision('Dev-Queue', 'test.txt', 1)
 
 		await P4Util.editFileAndSubmit(this.gateClient, GATE_FILENAME + 'queue.json', JSON.stringify({Change:this.secondEditCl}))
-		// this.info('removing window, sleeping again')
-		// await system.sleep(10)
 	}
 
 	verify() {
+		this.info('sleeping again after removing window')
+		await system.sleep(10)
+
 		return Promise.all([
 			this.checkHeadRevision('Main', 'test.txt', 4),
 			this.checkHeadRevision('Dev-Exact', 'test.txt', 2),
