[NO NON-PII COMMIT MESSAGE]

diff --git a/Engine/Extras/RoboMerge/v3/functional_tests/framework/src/tests/test-edge-gate.ts b/Engine/Extras/RoboMerge/v3/functional_tests/framework/src/tests/test-edge-gate.ts
--- a/Engine/Extras/RoboMerge/v3/functional_tests/framework/src/tests/test-edge-gate.ts
+++ b/Engine/Extras/RoboMerge/v3/functional_tests/framework/src/tests/test-edge-gate.ts
@@ -63,13 +63,11 @@ export class TestEdgeGate extends FunctionalTest {
 
 	verify() {
 		this.info('sleeping again after removing window')
-		await system.sleep(10)
-
 		return Promise.all([
 			this.checkHeadRevision('Main', 'test.txt', 4),
 			this.checkHeadRevision('Dev-Exact', 'test.txt', 2),
 			this.checkHeadRevision('Dev-PlusOne', 'test.txt', 2),
-			this.checkHeadRevision('Dev-Queue', 'test.txt', 3),
+			system.sleep(10).then(() => this.checkHeadRevision('Dev-Queue', 'test.txt', 3)),
 			this.checkHeadRevision('Dev-NoGate', 'test.txt', 4)
 		])
 	}
