[NO NON-PII COMMIT MESSAGE]

diff --git a/Engine/Extras/RoboMerge/v3/README.md b/Engine/Extras/RoboMerge/v3/README.md
--- a/Engine/Extras/RoboMerge/v3/README.md
+++ b/Engine/Extras/RoboMerge/v3/README.md
@@ -116,3 +116,36 @@ for both node and edge.
 |`notify`                   |e |            | |Additional people to email on blockages  |Also `globalNotify`|
 |`resolver`                 |e |            | |Single designated resolver               |Currently applies to both source and target nodes|
 |`whitelist`                |e |            | |Only specified users can integrate using edge| |
+
+##### Integration window examples
+
+Pause CIS when catching up to a gate
+
+```
+{ "pauseCISUnlessAtGate": true }
+```
+
+Specify integration windows when a gate catch-up may begin
+Allow integration between 7pm and 5am UTC
+
+```
+{ "integrationWindow": [{ "startHourUTC": 19, "durationHours": 10 }] }
+```
+
+Allow integration 2am -> 5am and 9pm -> midnight UTC
+
+```
+{ "integrationWindow":
+  [ { "startHourUTC": 2, "durationHours": 3 }
+  , { "startHourUTC": 21, "durationHours": 3 }
+  ]
+}
+```
+
+Dis-allow integration between 5am and 9pm UTC
+
+```
+{ "integrationWindow": [{ "startHourUTC": 5, "durationHours": 16 }]
+, "invertIntegrationWindow": true
+}
+```
