[NO NON-PII COMMIT MESSAGE]

diff --git a/Engine/Extras/RoboMerge/v3/public/js/boilerplate.js b/Engine/Extras/RoboMerge/v3/public/js/boilerplate.js
--- a/Engine/Extras/RoboMerge/v3/public/js/boilerplate.js
+++ b/Engine/Extras/RoboMerge/v3/public/js/boilerplate.js
@@ -954,7 +954,7 @@ function renderNameCell_Common(data, botname) {
 			msg += `Will unpause and retry on: ${unpauseTime}.<br />`;
 		}
 		else if (data.lastGoodCL && data.headCL && data.lastGoodCL === data.last_cl && data.last_cl < data.headCL) {
-			msg = 'Waiting on CIS gate'
+			msg = data.gateClosedMessage || 'Waiting on CIS gate'
 		}
 
 		if (msg) {
diff --git a/Engine/Extras/RoboMerge/v3/src/robo/badges.ts b/Engine/Extras/RoboMerge/v3/src/robo/badges.ts
--- a/Engine/Extras/RoboMerge/v3/src/robo/badges.ts
+++ b/Engine/Extras/RoboMerge/v3/src/robo/badges.ts
@@ -5,7 +5,7 @@ import { Badge } from '../common/badge';
 import { ContextualLogger } from '../common/logger';
 import { PerforceContext } from '../common/perforce';
 import { Blockage, Branch, BranchGraphInterface, ChangeInfo } from './branch-interfaces';
-import { BeginIntegratingToGateEvent, GateEventContext } from './branch-interfaces';
+import { BeginIntegratingToGateEvent, EndIntegratingToGateEvent, GateEventContext } from './branch-interfaces';
 import { PersistentConflict, Resolution } from './conflict-interfaces';
 import { BotEventHandler, BotEvents } from './events';
 
@@ -145,7 +145,8 @@ class GateHandler implements BotEventHandler {
 	}
 
 	onBeginIntegratingToGate(arg: BeginIntegratingToGateEvent) {
-		this.logger.info(GateHandler.logStringFor(arg.context) + ' catching up to CL# ' + arg.info.cl)
+		const suffix = arg.changesRemaining > 0 ? ` (${arg.changesRemaining} behind)` : ''
+		this.logger.info(GateHandler.logStringFor(arg.context) + ` catching up to CL#${arg.info.cl}` + suffix)
 		if (arg.context.pauseCIS && !this.pausedFlows.has(arg.context.to.upperName)) {
 			this.pausedFlows.set(arg.context.to.upperName, setTimeout(() => {
 				// @todo fill this in to pause!
@@ -161,9 +162,10 @@ class GateHandler implements BotEventHandler {
 		}
 	}
 
-	onEndIntegratingToGate(arg: GateEventContext) {
-		this.logger.info(GateHandler.logStringFor(arg) + ' caught up')
-		const pauseRefresher = this.pausedFlows.get(arg.to.upperName)
+	onEndIntegratingToGate(arg: EndIntegratingToGateEvent) {
+		const suffix = arg.targetCl > 0 ? ` (-> #${arg.targetCl})` : ''
+		this.logger.info(GateHandler.logStringFor(arg.context) + ' caught up' + suffix)
+		const pauseRefresher = this.pausedFlows.get(arg.context.to.upperName)
 		if (pauseRefresher) {
 			// @todo fill this in to unpause!
 			const url = ''
@@ -172,10 +174,10 @@ class GateHandler implements BotEventHandler {
 				url,
 				body,
 				contentType: 'application/json'
-			}, `Unpause of CIS for ${arg.to}`)
+			}, `Unpause of CIS for ${arg.context.to}`)
 
 			clearTimeout(pauseRefresher)
-			this.pausedFlows.delete(arg.to.upperName)
+			this.pausedFlows.delete(arg.context.to.upperName)
 		}
 	}
 }
diff --git a/Engine/Extras/RoboMerge/v3/src/robo/branch-interfaces.ts b/Engine/Extras/RoboMerge/v3/src/robo/branch-interfaces.ts
--- a/Engine/Extras/RoboMerge/v3/src/robo/branch-interfaces.ts
+++ b/Engine/Extras/RoboMerge/v3/src/robo/branch-interfaces.ts
@@ -253,4 +253,12 @@ export type GateEventContext = {
 export type BeginIntegratingToGateEvent = {
 	context: GateEventContext
 	info: GateInfo
+
+	// would like to provide this for logging, but I'm not storing it in a convenient place yet
+	changesRemaining: number
+}
+
+export type EndIntegratingToGateEvent = {
+	context: GateEventContext
+	targetCl: number
 }
diff --git a/Engine/Extras/RoboMerge/v3/src/robo/edgebot.ts b/Engine/Extras/RoboMerge/v3/src/robo/edgebot.ts
--- a/Engine/Extras/RoboMerge/v3/src/robo/edgebot.ts
+++ b/Engine/Extras/RoboMerge/v3/src/robo/edgebot.ts
@@ -13,7 +13,7 @@ import { PersistentConflict } from "./conflict-interfaces";
 import { BotEventTriggers } from "./events";
 import { NodeBot } from "./nodebot";
 import { isUserAKnownBot, postMessageToChannel } from "./notifications";
-import { PerforceRequestResult, PerforceStatefulBot } from "./perforce-stateful-bot";
+import { PerforceStatefulBot } from "./perforce-stateful-bot";
 import { Context } from "./settings";
 import { SlackMessageStyles } from "./slack";
 import { EdgeOptions } from "./branchdefs";
@@ -39,6 +39,8 @@ const MAX_CONFLICTS_TO_LIST = 5
 
 export type ResolveResultDetail = 'quick' | 'detailed'
 
+export type EdgeIntegrationResult = string[] | Change
+
 class EdgeBotImpl extends PerforceStatefulBot {
 	readonly graphBotName: string
 	readonly branch: Branch
@@ -51,8 +53,6 @@ class EdgeBotImpl extends PerforceStatefulBot {
 
 	// would like to encapsulate this in EdgeBot, but resolution is currently done by node bot
 	private currentIntegrationStartTimestamp = -1
-	// just used for status: how many changes remaining if we're catching up on multiple commits
-	private numChangesRemaining = 0
 
 	/**
 	 * Returns true if bot is currently merging a change
@@ -130,12 +130,12 @@ class EdgeBotImpl extends PerforceStatefulBot {
 	}
 
 	async setBotToLatestCl(): Promise<void> {
-		this.numChangesRemaining = 0
+		this.gate.numChangesRemaining = 0
 		return PerforceStatefulBot.setBotToLatestClInBranch(this, this.sourceBranch)
 	}
 
 	protected _forceSetLastCl_NoReset(value: number) {
-		this.numChangesRemaining = 0
+		this.gate.numChangesRemaining = 0
 		return super._forceSetLastCl_NoReset(value)
 	}
 
@@ -167,22 +167,21 @@ class EdgeBotImpl extends PerforceStatefulBot {
 
 	}
 
-	updateLastCl(changesFetched: Change[], changeIndex: number) {
-		this.lastCl = changesFetched[changeIndex].change
-		this.gate.setLastCl(this.lastCl)
+	updateLastCl(changesFetched: Change[], changeIndex: number, targetCl?: number) {
+		this.gate.updateLastCl(changesFetched, changeIndex, targetCl)
 
-		this.numChangesRemaining = this.gate.getNumChangesRemaining(changesFetched, changeIndex)
+		this.lastCl = this.gate.lastCl
 	}
 
-	private async getPerforceRequestResultFromCL(changelist: number, path?: string, changelistStatus?: ChangelistStatus) : Promise<PerforceRequestResult> {
+	private async getPerforceRequestResultFromCL(changelist: number, path?: string, changelistStatus?: ChangelistStatus) : Promise<EdgeIntegrationResult> {
 		const result = await this._getChange(changelist, path, changelistStatus)
 
 		if (!result.changes || !result.changes[0]) {
 			// swallowing actual error! should probably fix
-			return { errors : [`Failed to get ${changelistStatus ? changelistStatus : 'submitted'} CL ${changelist} (errors = ${result.errors}`] }
+			return [`Failed to get ${changelistStatus ? changelistStatus : 'submitted'} CL ${changelist} (errors = ${result.errors}`]
 		}
 
-		return { changes : result.changes, errors : [] }
+		return result.changes[0]
 	}
 	
 	private analyzeConflict(unresolved: ConflictedResolveNFile[]) {
@@ -276,7 +275,7 @@ class EdgeBotImpl extends PerforceStatefulBot {
 		this.block(this.createEdgeBlockageInfo(failure, pending), pauseDurationSeconds)
 	}
 
-	performMerge(info: ChangeInfo, target: MergeAction, convertIntegratesToEdits: boolean): Promise<PerforceRequestResult> {
+	performMerge(info: ChangeInfo, target: MergeAction, convertIntegratesToEdits: boolean): Promise<EdgeIntegrationResult> {
 		// make sure we come back in here afterwords
 		this._log_action(`Merging CL ${info.cl} via ${this.fullName} (${target.mergeMode})`)
 
@@ -408,7 +407,7 @@ class EdgeBotImpl extends PerforceStatefulBot {
 		this.currentIntegrationStartTimestamp = -1
 	}
 
-	private async integrate(info: ChangeInfo, target: MergeAction) : Promise<PerforceRequestResult> {
+	private async integrate(info: ChangeInfo, target: MergeAction) : Promise<EdgeIntegrationResult> {
 		let to_integrate = info.cl
 		this._log_action(`Integrating CL ${to_integrate} to ${this.targetBranch.name}`)
 
@@ -458,7 +457,7 @@ class EdgeBotImpl extends PerforceStatefulBot {
 				// integration not necessary
 				this.edgeBotLogger.info(msg)
 				await this.p4.deleteCl(this.targetBranch.workspace, changenum)
-				return { errors: [ msg ] }
+				return [ msg ]
 			}
 		}
 
@@ -483,7 +482,7 @@ class EdgeBotImpl extends PerforceStatefulBot {
 		await this.handleIntegrationError(failure, pending)
 		// Send to source node to facilitate notification handling
 		await this.sourceNode.handleMergeFailure(failure, pending)
-		return { errors }
+		return errors
 	}
 
 	private async revertAndDelete(cl: number) {
@@ -510,7 +509,7 @@ class EdgeBotImpl extends PerforceStatefulBot {
 	}
 
 // code restored that was deleted in CL#8057353 (from nodebot.ts) - still needs work
-	private async roboShelve(result: ResolveResult, toShelve: PendingChange) : Promise<PerforceRequestResult> {
+	private async roboShelve(result: ResolveResult, toShelve: PendingChange) : Promise<EdgeIntegrationResult> {
 		const targetBranch = toShelve.action.branch
 		const targetEdge = this
 		let errorStr: string | null = null
@@ -549,7 +548,7 @@ class EdgeBotImpl extends PerforceStatefulBot {
 				await this.sourceNode.handleMergeFailure(failure, toShelve)
 			}
 			
-			return { errors : [`${errorStr}: ${errorDetails}`] }
+			return [`${errorStr}: ${errorDetails}`]
 		}
 		else {
 
@@ -567,9 +566,8 @@ class EdgeBotImpl extends PerforceStatefulBot {
 	 * Perform a resolve of the change described by pending
 	 * @param {PendingChange} pending Pending change to resolve and submit
 	 * @param {number} [submitRetries = 3] number of times 'targetEdge.submitChangelist' can fail gracefully and come back to this method
-	 * @todo In a perfect world, parameter 'submitRetries' wouldn't exist and an orchastration method would handle this in a loop.
 	 */
-	async _resolveChangelist(pending: PendingChange, submitRetries  = 3) : Promise<PerforceRequestResult> {
+	async _resolveChangelist(pending: PendingChange, submitRetries  = 3) : Promise<EdgeIntegrationResult> {
 		// do a resolve with P4
 		this._log_action(`Resolving CL ${pending.change.cl} against ${this.targetBranch.name}`)
 		let detail: ResolveResultDetail = 'detailed'
@@ -622,7 +620,7 @@ class EdgeBotImpl extends PerforceStatefulBot {
 			await this.handlePostIntegrationFailure(failure, pending)
 			await this.sourceNode.handleMergeFailure(failure, pending, true)
 
-			return { errors: [`${failure.kind}: ${failure.description}`] }
+			return [`${failure.kind}: ${failure.description}`]
 		}
 
 		return this.getPerforceRequestResultFromCL(pending.newCl)
@@ -857,7 +855,7 @@ class EdgeBotImpl extends PerforceStatefulBot {
 		status.rootPath = this.targetBranch.rootPath
 
 		status.last_cl = this.lastCl
-		status.num_changes_remaining = this.numChangesRemaining
+		status.num_changes_remaining = this.gate.numChangesRemaining
 
 		status.is_active = this.isActive
 		status.is_available = this.isAvailable
@@ -1000,8 +998,8 @@ export class EdgeBot extends EdgeBotEntryPoints {
 		this.impl.preIntegrate(cl)
 	}
 
-	updateLastCl(changesFetched: Change[], changeIndex: number) {
-		this.impl.updateLastCl(changesFetched, changeIndex)
+	updateLastCl(changesFetched: Change[], changeIndex: number, targetCl?: number) {
+		this.impl.updateLastCl(changesFetched, changeIndex, targetCl)
 	}
 
 	/* Mirrored Variables */
diff --git a/Engine/Extras/RoboMerge/v3/src/robo/events.ts b/Engine/Extras/RoboMerge/v3/src/robo/events.ts
--- a/Engine/Extras/RoboMerge/v3/src/robo/events.ts
+++ b/Engine/Extras/RoboMerge/v3/src/robo/events.ts
@@ -1,7 +1,7 @@
 // Copyright Epic Games, Inc. All Rights Reserved.
 
 import { AlreadyIntegrated, Blockage, ChangeInfo, ForcedCl, PendingChange } from './branch-interfaces';
-import { BeginIntegratingToGateEvent, GateEventContext } from './branch-interfaces';
+import { BeginIntegratingToGateEvent, EndIntegratingToGateEvent } from './branch-interfaces';
 import { BotConfig } from './branchdefs';
 import { PersistentConflict } from './conflict-interfaces';
 
@@ -15,7 +15,7 @@ type OnNonSkipLastClChange = (arg: ForcedCl) => void
 type OnBranchUnblocked = (arg: PersistentConflict) => void
 type OnConflictStatus = (arg: boolean) => void
 type OnBeginIntegratingToGate = (arg: BeginIntegratingToGateEvent) => void
-type OnEndIntegratingToGate = (arg: GateEventContext) => void
+type OnEndIntegratingToGate = (arg: EndIntegratingToGateEvent) => void
 
 export interface BotEventHandler {
 	onChangeParsed?: OnChangeParsed
@@ -186,7 +186,7 @@ export class BotEventTriggers extends BotEvents {
 		}
 	}
 
-	reportEndIntegratingToGate(arg: GateEventContext) {
+	reportEndIntegratingToGate(arg: EndIntegratingToGateEvent) {
 		for (const listener of this.endIntegratingToGateListeners) {
 			listener(arg)
 		}
diff --git a/Engine/Extras/RoboMerge/v3/src/robo/gate.ts b/Engine/Extras/RoboMerge/v3/src/robo/gate.ts
--- a/Engine/Extras/RoboMerge/v3/src/robo/gate.ts
+++ b/Engine/Extras/RoboMerge/v3/src/robo/gate.ts
@@ -246,7 +246,8 @@ export class Gate {
 			if (!this.tryNextGate(cl)) {
 				this.reportCaughtUp()
 				if (this.currentGateInfo) {
-					adjustedCl = this.lastCl = this.currentGateInfo.cl
+					adjustedCl = this.currentGateInfo.cl
+					this.setLastCl(adjustedCl)
 				}
 			}
 		}
@@ -290,7 +291,9 @@ export class Gate {
 		return true
 	}
 
-	setLastCl(cl: number) {
+	updateLastCl(changesFetched: Change[], changeIndex: number, targetCl?: number) {
+		const cl = changesFetched[changeIndex].change
+
 		let notifyCaughtUp = false
 		if (this.currentGateInfo) {
 			// ignore going backwards; did we reach the gate?
@@ -302,13 +305,17 @@ export class Gate {
 			}
 		}
 
-		this.lastCl = cl
+		this.setLastCl(cl)
 		if (notifyCaughtUp) {
-			this.reportCaughtUp()
+			this.reportCaughtUp(targetCl)
 		}
+
+		this.numChangesRemaining = this.calcNumChangesRemaining(changesFetched, changeIndex)
 	}
 
-	getNumChangesRemaining(changesFetched: Change[], changeIndex: number) {
+	numChangesRemaining = 0
+
+	private calcNumChangesRemaining(changesFetched: Change[], changeIndex: number) {
 		const mostRecentGate = this.getMostRecentGate()
 		if (!mostRecentGate) {
 			return changesFetched.length - changeIndex - 1
@@ -428,12 +435,16 @@ export class Gate {
 		}
 		this.eventContext.eventTriggers.reportBeginIntegratingToGate({
 			context: this.eventContext,
-			info: mostRecent
+			info: mostRecent,
+			changesRemaining: this.numChangesRemaining
 		})
 	}
 
-	private reportCaughtUp() {
-		this.eventContext.eventTriggers.reportEndIntegratingToGate(this.eventContext)
+	private reportCaughtUp(targetCl?: number) {
+		this.eventContext.eventTriggers.reportEndIntegratingToGate({
+			context: this.eventContext,
+			targetCl: targetCl || -1
+		})
 	}
 
 	private persist() {
@@ -470,11 +481,15 @@ export class Gate {
 		}
 	}
 
-	private get lastCl() {
+	getEventContextForTests() {
+		return this.eventContext
+	}
+
+	get lastCl() {
 		return this.eventContext.edgeLastCl
 	}
 
-	private set lastCl(cl: number) {
+	private setLastCl(cl: number) {
 		this.eventContext.edgeLastCl = cl
 	}
 
@@ -487,6 +502,9 @@ const colors = require('colors')
 colors.enable()
 colors.setTheme(require('colors/themes/generic-logging.js'))
 
+function setLastCl(gate: Gate, cl: number) {
+	gate.updateLastCl([{change: cl, client: '', user: '', desc: ''}], 0)
+}
 
 export async function runTests(parentLogger: ContextualLogger) {
 	const logger = parentLogger.createChild('Gate')
@@ -535,13 +553,13 @@ export async function runTests(parentLogger: ContextualLogger) {
 		assert(et.beginCalls === 1, 'catching up')
 		assert(gate.isGateOpen(), 'gate open')
 
-		gate.setLastCl(2)
+		setLastCl(gate, 2)
 		await gate.tick()
 
 		if (!exact) {
 			const newCl = gate.preIntegrate(4) // only know we've caught up when higher cl comes in
 			assert(newCl === 3, 'last cl adjustment requested')
-			gate.setLastCl(3)
+			setLastCl(gate, 3)
 		}
 		assert(et.beginCalls === 1 && et.endCalls === 1, 'caught up')
 	}
@@ -603,7 +621,7 @@ export async function runTests(parentLogger: ContextualLogger) {
 		for (const cl of incoming) {
 
 			gate.preIntegrate(cl)
-			gate.setLastCl(cl)
+			setLastCl(gate, cl)
 			await gate.tick()
 		}
 
@@ -629,11 +647,11 @@ export async function runTests(parentLogger: ContextualLogger) {
 		assert(gate.isGateOpen(), 'still open after gate queued')
 
 		// start integrating
-		gate.preIntegrate(1); gate.setLastCl(1); await gate.tick()
+		gate.preIntegrate(1); setLastCl(gate, 1); await gate.tick()
 		if (middleIntegration) {
-			gate.preIntegrate(2); gate.setLastCl(2); await gate.tick()
+			gate.preIntegrate(2); setLastCl(gate, 2); await gate.tick()
 		}
-		gate.preIntegrate(3); gate.setLastCl(3); await gate.tick()
+		gate.preIntegrate(3); setLastCl(gate, 3); await gate.tick()
 
 		assert(et.beginCalls === 1 && et.endCalls === 1, 'catch-up notified')
 
@@ -645,7 +663,7 @@ export async function runTests(parentLogger: ContextualLogger) {
 
 	await gate1.tick()
 	gate1.preIntegrate(2)
-	gate1.setLastCl(2)
+	setLastCl(gate1, 2)
 	assert(et1.beginCalls + et1.endCalls === 0, 'no events')
 	assert(gate1.isGateOpen(), 'gate open')
 
@@ -705,13 +723,13 @@ export async function runTests(parentLogger: ContextualLogger) {
 		assert(gate.isGateOpen(), 'gate open')
 
 		gate.preIntegrate(5)
-		gate.setLastCl(5)
+		setLastCl(gate, 5)
 		await gate.tick()
 		assert(et.beginCalls === 1, 'no more events')
 		assert(gate.isGateOpen(), 'gate open')
 
 		gate.preIntegrate(7)
-		gate.setLastCl(7)
+		setLastCl(gate, 7)
 		assert(et.beginCalls === 1 && et.endCalls === 1, 'caught up')
 		assert(!gate.isGateOpen() && gate.getGateClosedMessage()!.includes('CIS'), 'gate closed')
 
diff --git a/Engine/Extras/RoboMerge/v3/src/robo/nodebot.ts b/Engine/Extras/RoboMerge/v3/src/robo/nodebot.ts
--- a/Engine/Extras/RoboMerge/v3/src/robo/nodebot.ts
+++ b/Engine/Extras/RoboMerge/v3/src/robo/nodebot.ts
@@ -4,15 +4,15 @@ import * as Sentry from '@sentry/node';
 import { _nextTick } from '../common/helper';
 import { ContextualLogger } from '../common/logger';
 import { Mailer, MailParams, Recipients } from '../common/mailer';
-import { Change, ConflictedResolveNFile, PerforceContext, ResolveResult, RoboWorkspace } from '../common/perforce';
+import { Change, ConflictedResolveNFile, PerforceContext, RoboWorkspace } from '../common/perforce';
 import { IPCControls, NodeBotInterface } from './bot-interfaces';
 import { BlockagePauseInfo, BranchStatus, QueuedChange, ReconsiderArgs } from './status-types';
 import { AlreadyIntegrated, Blockage, Branch, BranchArg, BranchGraphInterface, ChangeInfo, Failure } from './branch-interfaces';
 import { MergeAction, OperationResult, PendingChange, resolveBranchArg, StompedRevision, StompVerification, StompVerificationFile }  from './branch-interfaces';
 import { Conflicts } from './conflicts';
-import { EdgeBot } from './edgebot';
+import { EdgeBot, EdgeIntegrationResult } from './edgebot';
 import { BotEventTriggers } from './events';
-import { PerforceRequestResult, PerforceStatefulBot } from './perforce-stateful-bot';
+import { PerforceStatefulBot } from './perforce-stateful-bot';
 import { BlockageNodeOpUrls, OperationUrlHelper } from './roboserver';
 import { Context } from './settings';
 import { PauseState } from './state-interfaces';
@@ -47,8 +47,44 @@ function isStompableNonBinary(filename: string) {
 
 type InfoResult = { info?: ChangeInfo, errors: any[] }
 
-interface EdgeMergeResults extends PerforceRequestResult {
-	skippedEdges?: Set<string>
+type EdgeMerge = 'ok' | 'skipped' | 'error'
+
+/** results from integrating along available outgoing edges */
+class EdgeMergeResults {
+	// trying out making this purely per edge
+
+	addIntegrationResult(action: MergeAction, result: EdgeIntegrationResult) {
+		if ((result as Change).change) {
+			this.edges.set(action.branch.name, {merge: 'ok', change: result as Change})
+		}
+		else {
+			this.edges.set(action.branch.name, {merge: 'error'})
+			this.errors = [...this.errors, ...(result as string[])]
+		}
+	}
+
+	allChanges() {
+		return [...this.edges].map(([_, v]) => v.change).filter(Boolean) as Change[]
+	}
+
+	markSkipped(target: Branch) {
+		this.edges.set(target.name, {merge: 'skipped'})
+	}
+
+	wasSkipped(target: Branch) {
+		return this.edges.has(target.name) && this.edges.get(target.name)!.merge === 'skipped'
+	}
+
+	/** shelved or committed cl */
+	getTargetCl(target: Branch) {
+		const val = this.edges.get(target.name)
+		return val && val.change && val.change.change
+	}
+
+	// tuples of target name, result, opt change
+	edges = new Map<string, {merge: EdgeMerge, change?: Change}>()
+	// accumulated errors for all edges
+	errors: string[] = []
 }
 
 
@@ -577,12 +613,9 @@ export class NodeBot extends PerforceStatefulBot implements NodeBotInterface {
 
 		// Process and attempt the integration
 		const edgeMap = new Map<string, EdgeBot>([[edge.targetBranch.upperName, edge]])
-		let processResult: InfoResult
-		let mergeResult: PerforceRequestResult | undefined // Set to possible undefines as we use this in the finally clause
-		let integratedCl: Change
-		let resolveResult: ResolveResult
+		let mergeResult: EdgeMergeResults | undefined // Set to possible undefines as we use this in the finally clause
 		try {
-			processResult = await this._createChangeInfo(blockageChange, edgeMap, this.p4.username, targetBranch.workspace, targetBranch)
+			const processResult = await this._createChangeInfo(blockageChange, edgeMap, this.p4.username, targetBranch.workspace, targetBranch)
 			if (!processResult.info) {
 				return {
 					success: false, 
@@ -592,8 +625,9 @@ export class NodeBot extends PerforceStatefulBot implements NodeBotInterface {
 			mergeResult = await this._mergeClViaEdges(edgeMap, processResult.info, true)
 
 			// If we don't have a change returned, we ran into an error processing the CL
-			if (!mergeResult.changes || mergeResult.changes.length !== 1) { // Encountered when we get errors back
-				let message = `Got incorrect number of entries back from integrating CL ${changeCl} (length = ${mergeResult.changes ? mergeResult.changes.length : 'null'}, should be 1)`
+			const changes = mergeResult.allChanges()
+			if (changes.length !== 1) { // Encountered when we get errors back
+				let message = `Got incorrect number of entries back from integrating CL ${changeCl} (length = ${changes.length}, should be 1)`
 				if (mergeResult.errors.length > 0) {
 					message = `${message}\nErrors: ${mergeResult.errors}`
 				}
@@ -603,7 +637,7 @@ export class NodeBot extends PerforceStatefulBot implements NodeBotInterface {
 				}
 			}
 
-			integratedCl = mergeResult.changes[0]
+			const integratedCl = changes[0]!
 			if (!integratedCl.change) {
 				return { 
 					success: false, 
@@ -620,7 +654,7 @@ export class NodeBot extends PerforceStatefulBot implements NodeBotInterface {
 			}
 
 			// Get what is left to resolve out of the changelist
-			resolveResult = await this.p4.resolve(integratedCl.client, integratedCl.change, 'normal')
+			const resolveResult = await this.p4.resolve(integratedCl.client, integratedCl.change, 'normal')
 
 			if (!resolveResult.hasConflict()) {
 				return { 
@@ -807,9 +841,9 @@ export class NodeBot extends PerforceStatefulBot implements NodeBotInterface {
 		finally {
 			// Delete the CL, including any remaining shelved files
 			this.nodeBotLogger.info(`${this.fullName} - Cleaning up after verifyStompRequest()`)
-			if (mergeResult && mergeResult.changes && mergeResult.changes.length > 0) {
-				for (let processResultChange of mergeResult.changes) {
-					edge.revertPendingCLWithShelf(processResultChange.client, processResultChange.change, "Stomp Verify")
+			if (mergeResult) {
+				for (const processResultChange of mergeResult.allChanges()) {
+					await edge.revertPendingCLWithShelf(processResultChange.client, processResultChange.change, "Stomp Verify")
 				}
 			}
 		}
@@ -1265,8 +1299,8 @@ export class NodeBot extends PerforceStatefulBot implements NodeBotInterface {
 			// Every edge that was available at the start of the tick that is still available had the chance to receive this change, even if they weren't a target
 			// Count this as processed to keep them up-to-date
 			for (const edgeBot of availableEdges.values()) {
-				if (edgeBot.isAvailable && edgeBot.lastCl < change.change && !(changeResult.skippedEdges && changeResult.skippedEdges.has(edgeBot.targetBranch.name))) {
-					edgeBot.updateLastCl(changes, changeIndex) 
+				if (edgeBot.isAvailable && edgeBot.lastCl < change.change && !changeResult.wasSkipped(edgeBot.targetBranch)) {
+					edgeBot.updateLastCl(changes, changeIndex, changeResult.getTargetCl(edgeBot.targetBranch))
 				}
 			}
 
@@ -1296,7 +1330,9 @@ export class NodeBot extends PerforceStatefulBot implements NodeBotInterface {
 		if (result.info) {
 			return await this._mergeClViaEdges(availableEdges, result.info, ignoreEdgePauseState)
 		}
-		return { errors: result.errors }
+		const emResult = new EdgeMergeResults()
+		emResult.errors = result.errors
+		return emResult
 	}
 
 	private async _createChangeInfo(
@@ -1621,7 +1657,7 @@ export class NodeBot extends PerforceStatefulBot implements NodeBotInterface {
 	}
 
 	private async _mergeClViaEdges(availableEdges: Map<string, EdgeBot>, info: ChangeInfo, ignoreEdgePauseState: boolean) : Promise<EdgeMergeResults> {
-		let mergedChanges : EdgeMergeResults = { skippedEdges: new Set(), changes: [], errors: [] }
+		const mergedChanges = new EdgeMergeResults()
 
 		for (const target of info.targets!) {
 			if (!this.hasEdge(target.branch)) {
@@ -1657,7 +1693,7 @@ export class NodeBot extends PerforceStatefulBot implements NodeBotInterface {
 
 				// If this edge is paused, skip it but we'll have to report it so the node doesn't advance its lastCl
 				if (!targetEdge.isAvailable) {
-					mergedChanges.skippedEdges!.add(target.branch.name)
+					mergedChanges.markSkipped(target.branch)
 					availableEdges.delete(targetEdge.targetBranch.upperName)
 					continue
 				}
@@ -1673,7 +1709,7 @@ export class NodeBot extends PerforceStatefulBot implements NodeBotInterface {
 				++this.tickJournal.merges
 			}
 
-			const integration = await targetEdge.performMerge(info, target, this.branch.convertIntegratesToEdits)
+			const integrationResult = await targetEdge.performMerge(info, target, this.branch.convertIntegratesToEdits)
 
 			// wake up target nodebot, so look for further merges or conflict resolutions
 			const targetNodebot = target.branch.bot as NodeBot
@@ -1688,10 +1724,7 @@ export class NodeBot extends PerforceStatefulBot implements NodeBotInterface {
 				availableEdges.delete(targetEdge.targetBranch.upperName)
 			}
 
-			if (integration.changes) {
-				mergedChanges.changes = mergedChanges.changes!.concat(integration.changes)
-			}
-			mergedChanges.errors = mergedChanges.errors.concat(integration.errors)
+			mergedChanges.addIntegrationResult(target, integrationResult)
 		}
 
 		return mergedChanges
