[NO NON-PII COMMIT MESSAGE]

diff --git a/Engine/Extras/RoboMerge/v3/public/js/boilerplate.js b/Engine/Extras/RoboMerge/v3/public/js/boilerplate.js
--- a/Engine/Extras/RoboMerge/v3/public/js/boilerplate.js
+++ b/Engine/Extras/RoboMerge/v3/public/js/boilerplate.js
@@ -953,8 +953,8 @@ function renderNameCell_Common(data, botname) {
 		if (data.is_paused || data.is_blocked) {
 			msg += `Will unpause and retry on: ${unpauseTime}.<br />`;
 		}
-		else if (data.lastGoodCL && data.headCL && data.lastGoodCL === data.last_cl && data.last_cl < data.headCL) {
-			msg = data.gateClosedMessage || 'Waiting on CIS gate'
+		else if (data.gateClosedMessage) {
+			msg = data.gateClosedMessage
 		}
 
 		if (msg) {
diff --git a/Engine/Extras/RoboMerge/v3/src/robo/gate.ts b/Engine/Extras/RoboMerge/v3/src/robo/gate.ts
--- a/Engine/Extras/RoboMerge/v3/src/robo/gate.ts
+++ b/Engine/Extras/RoboMerge/v3/src/robo/gate.ts
@@ -247,6 +247,7 @@ export class Gate {
 				this.reportCaughtUp()
 				if (this.currentGateInfo) {
 					adjustedCl = this.currentGateInfo.cl
+					this.context.logger.info(`Adjusting cl to ${adjustedCl}`)
 					this.setLastCl(adjustedCl)
 				}
 			}
@@ -354,7 +355,7 @@ export class Gate {
 		}
 
 		// null means no gate
-		return this.queuedGates.length > 0 ? 'waiting for window' : null
+		return this.queuedGates.length > 0 ? 'waiting for integration window' : null
 	}
 
 	applyStatus(outStatus: { [key: string]: any }) {
