[NO NON-PII COMMIT MESSAGE]

diff --git a/Engine/Extras/RoboMerge/v3/functional_tests/framework/src/tests/test-edge-gate.ts b/Engine/Extras/RoboMerge/v3/functional_tests/framework/src/tests/test-edge-gate.ts
--- a/Engine/Extras/RoboMerge/v3/functional_tests/framework/src/tests/test-edge-gate.ts
+++ b/Engine/Extras/RoboMerge/v3/functional_tests/framework/src/tests/test-edge-gate.ts
@@ -7,7 +7,8 @@ const streams: Stream[] = [
 	{name: 'Dev-NoGate', streamType: 'development', parent: 'Main'},
 	{name: 'Dev-PlusOne', streamType: 'development', parent: 'Main'},
 	{name: 'Dev-Exact', streamType: 'development', parent: 'Main'},
-	{name: 'Dev-Queue', streamType: 'development', parent: 'Main'}
+	{name: 'Dev-Queue', streamType: 'development', parent: 'Main'},
+	{name: 'Dev-Window', streamType: 'development', parent: 'Main'}
 ]
 
 const GATE_FILENAME = 'TestEdgeGate-gate'
@@ -25,12 +26,9 @@ export class TestEdgeGate extends FunctionalTest {
 		await P4Util.addFileAndSubmit(mainClient, 'test.txt', 'Initial content')
 
 		const desc = 'Initial branch of files from Main'
-		await Promise.all([
-			this.p4.populate(this.getStreamPath('Dev-Exact'), desc),
-			this.p4.populate(this.getStreamPath('Dev-PlusOne'), desc),
-			this.p4.populate(this.getStreamPath('Dev-NoGate'), desc),
-			this.p4.populate(this.getStreamPath('Dev-Queue'), desc)
-		])
+		await Promise.all(
+			['Exact', 'PlusOne', 'NoGate', 'Queue', 'Window'].map(s => this.p4.populate(this.getStreamPath('Dev-' + s), desc))
+		)
 
 		const firstEditCl = await P4Util.editFileAndSubmit(mainClient, 'test.txt', 'Initial content\n\nFirst addition')
 
@@ -38,7 +36,8 @@ export class TestEdgeGate extends FunctionalTest {
 		await Promise.all([
 			P4Util.addFile(this.gateClient, GATE_FILENAME + 'exact.json', JSON.stringify({Change:firstEditCl})),
 			P4Util.addFile(this.gateClient, GATE_FILENAME + 'plusone.json', JSON.stringify({Change:firstEditCl + 1})),
-			P4Util.addFile(this.gateClient, GATE_FILENAME + 'queue.json', JSON.stringify({Change:firstEditCl, integrationWindow: []}))
+			P4Util.addFile(this.gateClient, GATE_FILENAME + 'queue.json', JSON.stringify({Change:firstEditCl, integrationWindow: []})),
+			P4Util.addFile(this.gateClient, GATE_FILENAME + 'window.json', JSON.stringify({Change:firstEditCl, integrationWindow: []}))
 		])
 
 		await this.gateClient.submit('Added gates')
@@ -52,49 +51,43 @@ export class TestEdgeGate extends FunctionalTest {
 	async run() {
 		await this.waitForRobomergeIdle()
 		await this.checkHeadRevision('Dev-Queue', 'test.txt', 1)
-		await P4Util.editFileAndSubmit(this.gateClient, GATE_FILENAME + 'queue.json', JSON.stringify({Change:this.secondEditCl, integrationWindow: []}))
-
-		this.info('updated gate but window closed; sleeping') // don't have a way to wait for robomerge to pick up the gate change
-		await system.sleep(10)
-		await this.checkHeadRevision('Dev-Queue', 'test.txt', 1)
-
 		await P4Util.editFileAndSubmit(this.gateClient, GATE_FILENAME + 'queue.json', JSON.stringify({Change:this.secondEditCl}))
+		await P4Util.editFileAndSubmit(this.gateClient, GATE_FILENAME + 'window.json', JSON.stringify({Change:this.secondEditCl, integrationWindow: [{
+			startHourUTC: ((new Date).getUTCHours() + 12) % 24,
+			durationHours: 1
+		}]}))
 	}
 
 	verify() {
-		this.info('sleeping again after removing window')
+		this.info('sleeping after removing window')
 		return Promise.all([
 			this.checkHeadRevision('Main', 'test.txt', 4),
 			this.checkHeadRevision('Dev-Exact', 'test.txt', 2),
 			this.checkHeadRevision('Dev-PlusOne', 'test.txt', 2),
 			system.sleep(10).then(() => this.checkHeadRevision('Dev-Queue', 'test.txt', 3)),
-			this.checkHeadRevision('Dev-NoGate', 'test.txt', 4)
+			this.checkHeadRevision('Dev-NoGate', 'test.txt', 4),
+			system.sleep(10).then(() => this.checkHeadRevision('Dev-Window', 'test.txt', 1))
 		])
 	}
 
 	getBranches() {
-		const mainSpec = this.makeForceAllBranchDef('Main', ['Dev-Exact', 'Dev-NoGate', 'Dev-PlusOne', 'Dev-Queue'])
+		const mainSpec = this.makeForceAllBranchDef('Main', ['Dev-Exact', 'Dev-NoGate', 'Dev-PlusOne', 'Dev-Queue', 'Dev-Window'])
 		mainSpec.initialCL = 1
 		return [
 			mainSpec,
 			this.makeForceAllBranchDef('Dev-Exact', []),
 			this.makeForceAllBranchDef('Dev-PlusOne', []),
 			this.makeForceAllBranchDef('Dev-NoGate', []),
-			this.makeForceAllBranchDef('Dev-Queue', [])
+			this.makeForceAllBranchDef('Dev-Queue', []),
+			this.makeForceAllBranchDef('Dev-Window', []),
 		]
 	}
 
 	getEdges() {
-		return [
-		  { from: this.fullBranchName('Main'), to: this.fullBranchName('Dev-Exact')
-		  , lastGoodCLPath: this.gateClient.stream + '/' + GATE_FILENAME + 'exact.json'
-		  }
-		, { from: this.fullBranchName('Main'), to: this.fullBranchName('Dev-PlusOne')
-		  , lastGoodCLPath: this.gateClient.stream + '/' + GATE_FILENAME + 'plusone.json'
-		  }
-		, { from: this.fullBranchName('Main'), to: this.fullBranchName('Dev-Queue')
-		  , lastGoodCLPath: this.gateClient.stream + '/' + GATE_FILENAME + 'queue.json'
-		  }
-		]
+		return ['Exact', 'PlusOne', 'Queue', 'Window'].map(s =>
+		  ({ from: this.fullBranchName('Main'), to: this.fullBranchName('Dev-' + s)
+		  , lastGoodCLPath: this.gateClient.stream + '/' + GATE_FILENAME + s.toLowerCase() + '.json'
+		  , initialCL: 1
+		  }))
 	}
 }
