[NO NON-PII COMMIT MESSAGE]

diff --git a/Engine/Extras/RoboMerge/v3/public/js/boilerplate.js b/Engine/Extras/RoboMerge/v3/public/js/boilerplate.js
--- a/Engine/Extras/RoboMerge/v3/public/js/boilerplate.js
+++ b/Engine/Extras/RoboMerge/v3/public/js/boilerplate.js
@@ -955,6 +955,12 @@ function renderNameCell_Common(data, botname) {
 		}
 		else if (data.gateClosedMessage) {
 			msg = data.gateClosedMessage
+			if (data.nextWindowOpenTime) {
+				const timestamp = Date.parse(data.nextWindowOpenTime)
+				if (!isNaN(timestamp)) {
+					msg += ` (opens: ${(new Date(timestamp)).toLocaleString()})`
+				}
+			}
 		}
 
 		if (msg) {
diff --git a/Engine/Extras/RoboMerge/v3/src/robo/badges.ts b/Engine/Extras/RoboMerge/v3/src/robo/badges.ts
--- a/Engine/Extras/RoboMerge/v3/src/robo/badges.ts
+++ b/Engine/Extras/RoboMerge/v3/src/robo/badges.ts
@@ -163,7 +163,7 @@ class GateHandler implements BotEventHandler {
 	}
 
 	onEndIntegratingToGate(arg: EndIntegratingToGateEvent) {
-		const suffix = arg.targetCl > 0 ? ` (-> #${arg.targetCl})` : ''
+		const suffix = arg.targetCl > 0 ? ` (-> @${arg.targetCl})` : ''
 		this.logger.info(GateHandler.logStringFor(arg.context) + ' caught up' + suffix)
 		const pauseRefresher = this.pausedFlows.get(arg.context.to.upperName)
 		if (pauseRefresher) {
diff --git a/Engine/Extras/RoboMerge/v3/src/robo/gate.ts b/Engine/Extras/RoboMerge/v3/src/robo/gate.ts
--- a/Engine/Extras/RoboMerge/v3/src/robo/gate.ts
+++ b/Engine/Extras/RoboMerge/v3/src/robo/gate.ts
@@ -277,7 +277,7 @@ export class Gate {
 
 		if (!this.nowIsWithinAllowedCatchupWindow()) {
 			// want to make this an info log, but would spam every tick at the moment
-			this.context.logger.verbose('delaying gate catch-up due to configured window')
+			// this.context.logger.verbose('delaying gate catch-up due to configured window')
 
 			// wait for next window before catching up with queued gates
 			this.currentGateInfo = null
@@ -372,6 +372,9 @@ export class Gate {
 		const closedMessage = this.getGateClosedMessage()
 		if (closedMessage) {
 			outStatus.gateClosedMessage = closedMessage
+			if (this.nextWindowOpenTime) {
+				outStatus.nextWindowOpenTime = this.nextWindowOpenTime
+			}
 		}
 	}
 
@@ -383,6 +386,8 @@ export class Gate {
 	}
 
 	private nowIsWithinAllowedCatchupWindow() {
+		this.nextWindowOpenTime = null
+
 		const mostRecentGate = this.getMostRecentGate()
 		let integrationWindow: IntegrationWindowPane[] | null = null
 		let invert = false
@@ -426,9 +431,25 @@ export class Gate {
 		if (invert) {
 			inWindow = !inWindow
 		}
+
+		if (!inWindow && !invert && integrationWindow.length === 1) {
+			// for now, only provide message if we have one non-inverted window
+			// more general solution to follow
+			// also not taking into account days of the week yet
+			const now = new Date
+			const d = now
+			d.setUTCHours(integrationWindow[0].startHourUTC, 0, 0, 0);
+			if (d < now) {
+				d.setUTCDate(d.getUTCDate() + 1)
+			}
+			this.nextWindowOpenTime = d
+		}
+
 		return inWindow
 	}
 
+	nextWindowOpenTime: Date | null = null
+
 	private reportCatchingUp() {
 		const mostRecent = this.getMostRecentGate()
 		if (!mostRecent) {
