[NO NON-PII COMMIT MESSAGE]

diff --git a/Engine/Extras/RoboMerge/v3/src/robo/branch-interfaces.ts b/Engine/Extras/RoboMerge/v3/src/robo/branch-interfaces.ts
--- a/Engine/Extras/RoboMerge/v3/src/robo/branch-interfaces.ts
+++ b/Engine/Extras/RoboMerge/v3/src/robo/branch-interfaces.ts
@@ -233,9 +233,12 @@ export function gatesSame(lhs: GateInfo | null, rhs: GateInfo | null) {
 	for (let n = 0; n < lhs.integrationWindow.length; ++n) {
 		const lhsWindow = lhs.integrationWindow[n]
 		const rhsWindow = rhs.integrationWindow[n]
+		const lhsDays = lhsWindow.daysOfTheWeek || []
+		const rhsDays = rhsWindow.daysOfTheWeek || []
 		if (lhsWindow.startHourUTC !== rhsWindow.startHourUTC ||
 				lhsWindow.durationHours !== rhsWindow.durationHours ||
-				lhsWindow.dayOfTheWeek !== rhsWindow.dayOfTheWeek) {
+				lhsDays.length !== rhsDays.length ||
+				!lhsDays.every((lhsDay, index) => lhsDay === rhsDays[index])) {
 			return false
 		}
 	}
diff --git a/Engine/Extras/RoboMerge/v3/src/robo/branchdefs.ts b/Engine/Extras/RoboMerge/v3/src/robo/branchdefs.ts
--- a/Engine/Extras/RoboMerge/v3/src/robo/branchdefs.ts
+++ b/Engine/Extras/RoboMerge/v3/src/robo/branchdefs.ts
@@ -61,8 +61,8 @@ export const DAYS_OF_THE_WEEK = ['sun', 'mon', 'tue', 'wed', 'thu', 'fri', 'sat'
 type DayOfTheWeek = 'sun' | 'mon' | 'tue' | 'wed' | 'thu' | 'fri' | 'sat'
 
 export type IntegrationWindowPane = {
-	// if day not specified, daily
-	dayOfTheWeek?: DayOfTheWeek
+	// if days not specified, daily
+	daysOfTheWeek?: DayOfTheWeek[]
 	startHourUTC: number
 	durationHours: number
 }
@@ -142,12 +142,15 @@ export interface BranchGraphDefinition {
 function validateCommonOptions(options: Partial<CommonOptionFields>) {
 	if (options.integrationWindow) {
 		for (const pane of options.integrationWindow) {
-			if (pane.dayOfTheWeek) {
-				const day = pane.dayOfTheWeek.slice(0, 3).toLowerCase()
-				if (DAYS_OF_THE_WEEK.indexOf(day) < 0) {
-					throw new Error(`Unknown day of the week ${pane.dayOfTheWeek}`)
+			if (pane.daysOfTheWeek) {
+				for (let index = 0; index < pane.daysOfTheWeek.length; ++index) { 
+					const dayStr = pane.daysOfTheWeek[index]
+					const day = dayStr.slice(0, 3).toLowerCase()
+					if (DAYS_OF_THE_WEEK.indexOf(day) < 0) {
+						throw new Error(`Unknown day of the week ${dayStr}`)
+					}
+					pane.daysOfTheWeek[index] = day as DayOfTheWeek
 				}
-				pane.dayOfTheWeek = day as DayOfTheWeek
 			}
 		}
 	}
diff --git a/Engine/Extras/RoboMerge/v3/src/robo/gate.ts b/Engine/Extras/RoboMerge/v3/src/robo/gate.ts
--- a/Engine/Extras/RoboMerge/v3/src/robo/gate.ts
+++ b/Engine/Extras/RoboMerge/v3/src/robo/gate.ts
@@ -405,25 +405,35 @@ export class Gate {
 			return true
 		}
 
-		const now = new Date;
+		const now = new Date
+		const nowHour = now.getUTCHours()
+		const nowHourInWeek = now.getUTCDay() * 24 + nowHour
 		let inWindow = false
 
+	outer:
 		for (const pane of integrationWindow) {
-			if (pane.dayOfTheWeek) {
-				const dayIndex = DAYS_OF_THE_WEEK.indexOf(pane.dayOfTheWeek)
-				if (dayIndex < 0 || dayIndex > 6) {
-					throw new Error('invalid day, should have been caught in validation')
+			if (pane.daysOfTheWeek) {
+				for (const day of pane.daysOfTheWeek) {
+					const dayIndex = DAYS_OF_THE_WEEK.indexOf(day)
+					if (dayIndex < 0 || dayIndex > 6) {
+						throw new Error('invalid day, should have been caught in validation')
+					}
+
+					const paneStart = dayIndex * 24 + pane.startHourUTC
+
+					// same logic as for daily below, except 168 hours per week
+					if ((nowHourInWeek + 168 - paneStart) % 168 < pane.durationHours) {
+						inWindow = true
+						break outer
+					}
 				}
-
-				// @todo check if getUTCDay matches
-				// (need to account for duration spanning one or more midnights)
 			}
 			else {
 				// subtlely going over midnight, e.g. start at 11pm for 4 hours
 				// say current time is 1am, we do (1 + (24 - 11)) % 24, and see that it is < 4
-				if ((now.getUTCHours() + 24 - pane.startHourUTC) % 24 < pane.durationHours) {
+				if ((nowHour + 24 - pane.startHourUTC) % 24 < pane.durationHours) {
 					inWindow = true
-					break
+					break outer
 				}
 			}
 		}
@@ -432,17 +442,32 @@ export class Gate {
 			inWindow = !inWindow
 		}
 
-		if (!inWindow && !invert && integrationWindow.length === 1) {
+		if (!inWindow && !invert && integrationWindow.length > 0) {
+			// handling everything except inverted ranges
+			let earliestHour = 1000 // just has to be two weeks or more
+			const consider = (dayIndex: number, pane: IntegrationWindowPane) => {
+				let hourInWeek = dayIndex * 24 + pane.startHourUTC
+				if (hourInWeek < nowHourInWeek) {
+					hourInWeek += 168
+				}
+				earliestHour = Math.min(earliestHour, hourInWeek)				
+			}
+			for (const pane of integrationWindow) {
+				if (pane.daysOfTheWeek) {
+					for (const day of pane.daysOfTheWeek) {
+						consider(DAYS_OF_THE_WEEK.indexOf(day), pane)
+					}
+				}
+				else {
+					consider(now.getUTCDay(), pane)
+				}
+			}
+
 			// for now, only provide message if we have one non-inverted window
 			// more general solution to follow
 			// also not taking into account days of the week yet
-			const now = new Date
-			const d = now
-			d.setUTCHours(integrationWindow[0].startHourUTC, 0, 0, 0);
-			if (d < now) {
-				d.setUTCDate(d.getUTCDate() + 1)
-			}
-			this.nextWindowOpenTime = d
+			this.nextWindowOpenTime = now
+			this.nextWindowOpenTime.setUTCHours(earliestHour - now.getUTCDay() * 24, 0, 0, 0);
 		}
 
 		return inWindow
