[NO NON-PII COMMIT MESSAGE]

diff --git a/Engine/Extras/RoboMerge/v3/src/robo/targets.ts b/Engine/Extras/RoboMerge/v3/src/robo/targets.ts
--- a/Engine/Extras/RoboMerge/v3/src/robo/targets.ts
+++ b/Engine/Extras/RoboMerge/v3/src/robo/targets.ts
@@ -374,7 +374,8 @@ export function processOtherBotTargets(
 		const target = '!-'.indexOf(arg[0]) < 0 ? arg : arg.substr(1)
 		const branch = targetBranchGraph.getBranch(target)
 		if (!branch) {
-			if (targetBranchGraph.config.branchNamesToIgnore.indexOf(target.toUpperCase()) < 0) {
+			if (!targetBranchGraph.config.macros[target.toLowerCase()] &&
+				targetBranchGraph.config.branchNamesToIgnore.indexOf(target.toUpperCase()) < 0) {
 				errors.push(`Branch '${target}' not found in ${bot}`)
 			}
 			continue
