[NO NON-PII COMMIT MESSAGE]

diff --git a/Engine/Extras/RoboMerge/v3/src/robo/gate.ts b/Engine/Extras/RoboMerge/v3/src/robo/gate.ts
--- a/Engine/Extras/RoboMerge/v3/src/robo/gate.ts
+++ b/Engine/Extras/RoboMerge/v3/src/robo/gate.ts
@@ -510,7 +510,7 @@ export class Gate {
 	}
 
 	private loadFromPersistence() {
-		if (!this.persistence) {
+		if (true) {
 			return
 		}
 		const saved = this.persistence.get(GATE_INFO_KEY)
