[NO NON-PII COMMIT MESSAGE]

diff --git a/Engine/Extras/RoboMerge/v3/src/robo/gate.ts b/Engine/Extras/RoboMerge/v3/src/robo/gate.ts
--- a/Engine/Extras/RoboMerge/v3/src/robo/gate.ts
+++ b/Engine/Extras/RoboMerge/v3/src/robo/gate.ts
@@ -510,22 +510,22 @@ export class Gate {
 	}
 
 	private loadFromPersistence() {
-		if (true) {
-			return
-		}
-		const saved = this.persistence.get(GATE_INFO_KEY)
-		if (saved) {
-			this.context.logger.info(`Restoring saved gate info: current ${saved.current && saved.current.cl}`)
-
-			if (saved.current) {
-				this.currentGateInfo = saved.current
-				this.reportCatchingUp()
-			}
-			if (saved.queued) {
-				this.context.logger.info('Queue: ' + saved.queued.map((info: GateInfo) => info.cl).join(', '))
-				this.queuedGates = saved.queued
-			}
-		}
+		// if (!this.persistence) {
+		// 	return
+		// }
+		// const saved = this.persistence.get(GATE_INFO_KEY)
+		// if (saved) {
+		// 	this.context.logger.info(`Restoring saved gate info: current ${saved.current && saved.current.cl}`)
+
+		// 	if (saved.current) {
+		// 		this.currentGateInfo = saved.current
+		// 		this.reportCatchingUp()
+		// 	}
+		// 	if (saved.queued) {
+		// 		this.context.logger.info('Queue: ' + saved.queued.map((info: GateInfo) => info.cl).join(', '))
+		// 		this.queuedGates = saved.queued
+		// 	}
+		// }
 	}
 
 	getEventContextForTests() {
