[NO NON-PII COMMIT MESSAGE]

diff --git a/Engine/Extras/RoboMerge/v3/src/robo/gate.ts b/Engine/Extras/RoboMerge/v3/src/robo/gate.ts
--- a/Engine/Extras/RoboMerge/v3/src/robo/gate.ts
+++ b/Engine/Extras/RoboMerge/v3/src/robo/gate.ts
@@ -506,26 +506,36 @@ export class Gate {
 		if (this.queuedGates.length > 0) {
 			data.queued = this.queuedGates
 		}
+		data.lastCl = this.lastCl
 		this.persistence.set(GATE_INFO_KEY, data)
 	}
 
 	private loadFromPersistence() {
-		// if (!this.persistence) {
-		// 	return
-		// }
-		// const saved = this.persistence.get(GATE_INFO_KEY)
-		// if (saved) {
-		// 	this.context.logger.info(`Restoring saved gate info: current ${saved.current && saved.current.cl}`)
-
-		// 	if (saved.current) {
-		// 		this.currentGateInfo = saved.current
-		// 		this.reportCatchingUp()
-		// 	}
+		if (!this.persistence) {
+			return
+		}
+		const saved = this.persistence.get(GATE_INFO_KEY)
+		if (saved) {
+			this.context.logger.info(`Restoring saved gate info: current ${saved.current && saved.current.cl}`)
+
+			if (saved.lastCl) {
+				this.setLastCl(saved.lastCl)
+			}
+
+			if (saved.current) {
+				this.currentGateInfo = saved.current
+
+				if (saved.lastCl && saved.lastCl > saved.current.cl) {
+					this.reportCatchingUp()
+				}
+			}
+
 		// 	if (saved.queued) {
+			// need to convert strings to dates for one thing
 		// 		this.context.logger.info('Queue: ' + saved.queued.map((info: GateInfo) => info.cl).join(', '))
 		// 		this.queuedGates = saved.queued
 		// 	}
-		// }
+		}
 	}
 
 	getEventContextForTests() {
