[NO NON-PII COMMIT MESSAGE]

diff --git a/Engine/Extras/RoboMerge/v3/README.md b/Engine/Extras/RoboMerge/v3/README.md
--- a/Engine/Extras/RoboMerge/v3/README.md
+++ b/Engine/Extras/RoboMerge/v3/README.md
@@ -115,17 +115,17 @@ for both node and edge.
 |`maxFilesPerIntegration`   |e |`-1`        | |Reject integrations with more files than this| |
 |`notify`                   |e |            | |Additional people to email on blockages  |Also `globalNotify`|
 |`resolver`                 |e |            | |Single designated resolver               |Currently applies to both source and target nodes|
-|`whitelist`                |e |            | |Only specified users can integrate using edge| |
 
 ##### Integration window examples
 
-Pause CIS when catching up to a gate
+###### Pause CIS when catching up to a gate
 
 ```
 { "pauseCISUnlessAtGate": true }
 ```
 
-Specify integration windows when a gate catch-up may begin
+###### Specify integration windows when a gate catch-up may begin
+
 Allow integration between 7pm and 5am UTC
 
 ```
@@ -135,17 +135,11 @@ Allow integration between 7pm and 5am UTC
 Allow integration 2am -> 5am and 9pm -> midnight UTC
 
 ```
-{ "integrationWindow":
-  [ { "startHourUTC": 2, "durationHours": 3 }
-  , { "startHourUTC": 21, "durationHours": 3 }
-  ]
-}
+{ "integrationWindow": [{ "startHourUTC": 2, "durationHours": 3 }, { "startHourUTC": 21, "durationHours": 3 }] }
 ```
 
 Dis-allow integration between 5am and 9pm UTC
 
 ```
-{ "integrationWindow": [{ "startHourUTC": 5, "durationHours": 16 }]
-, "invertIntegrationWindow": true
-}
+{ "integrationWindow": [{ "startHourUTC": 5, "durationHours": 16 }], "invertIntegrationWindow": true }
 ```
diff --git a/Engine/Extras/RoboMerge/v3/functional_tests/framework/src/BranchState.ts b/Engine/Extras/RoboMerge/v3/functional_tests/framework/src/BranchState.ts
--- a/Engine/Extras/RoboMerge/v3/functional_tests/framework/src/BranchState.ts
+++ b/Engine/Extras/RoboMerge/v3/functional_tests/framework/src/BranchState.ts
@@ -48,6 +48,10 @@ export class BranchState {
 		return !!this.rawBranchData.edges[targetBranchName.toUpperCase()]
 	}
 
+	getEdges(): EdgeState[] {
+		return Object.entries(this.rawBranchData.edges).map(([_, edgeData]) => new EdgeState(edgeData))
+	}
+
 	getLastCL() {
 		return this.rawBranchData.last_cl
 	}
@@ -68,8 +72,11 @@ export class BranchState {
 		return this.rawBranchData.status_msg
 	}
 
-	getEdgeState(targetBranchName: string): EdgeState {
+	getTickCount() {
+		return this.rawBranchData.tick_count
+	}
 
+	getEdgeState(targetBranchName: string): EdgeState {
 		const edgeJson = this.rawBranchData.edges[targetBranchName.toUpperCase()]
 
 		if (!edgeJson) {
diff --git a/Engine/Extras/RoboMerge/v3/functional_tests/framework/src/framework.ts b/Engine/Extras/RoboMerge/v3/functional_tests/framework/src/framework.ts
--- a/Engine/Extras/RoboMerge/v3/functional_tests/framework/src/framework.ts
+++ b/Engine/Extras/RoboMerge/v3/functional_tests/framework/src/framework.ts
@@ -131,7 +131,6 @@ type RobomergeBranchOptions = {
 	flowsTo: string[]
 	forceFlowTo: string[]
 	defaultFlow: string[]
-	whitelist: string[]
 	resolver: string | null
 	aliases: string[]
 	badgeProject: string | null
@@ -192,6 +191,8 @@ type EdgeOptionFields = {
 	// by default, specify when gate catch ups are allowed; can be inverted to disallow
 	integrationWindow: IntegrationWindowPane[]
 	invertIntegrationWindow: boolean
+
+	implicitCommands: string[]
 }
 
 export type EdgeProperties = Partial<EdgeOptionFields> & {
@@ -199,10 +200,11 @@ export type EdgeProperties = Partial<EdgeOptionFields> & {
 	to: string
 }
 
+const NUM_WAIT_INTERVALS = 20
 export async function retryWithBackoff<T extends {}>(desc: string, f: (last: boolean) => Promise<T | null>): Promise<T> {
 	let sleepTime = .5
-	for (let safety = 0; safety < 20; ++safety) {
-		const result = await f(safety === 19)
+	for (let safety = 0; safety < NUM_WAIT_INTERVALS; ++safety) {
+		const result = await f(safety === NUM_WAIT_INTERVALS - 1)
 		if (result) {
 			return result
 		}
@@ -271,6 +273,7 @@ export abstract class FunctionalTest {
 			['Owner', stream.owner || FunctionalTest.dfltStreamOpts.owner],
 			['Name', stream.name],
 			['Parent', stream.parent ? this.getStreamPath(stream.parent, stream.depotName) : FunctionalTest.dfltStreamOpts.parent],
+			['ParentView', 'inherit'], // required by p4 2021.1
 			['Type', stream.streamType],
 			['Description', '\n\tCreated by root.'],
 			// development streams given fromparent - let's see if it matters
@@ -353,7 +356,7 @@ export abstract class FunctionalTest {
 			.replace('<bot>', botName)
 			.replace('<branch>', nodeName.toUpperCase())
 
-		const errorPrefix = "test failure: couldn't query branch info: "
+		const errorPrefix = "couldn't query branch info: "
 
 		let json: { [key: string]: any }
 
@@ -363,15 +366,6 @@ export abstract class FunctionalTest {
 		catch (err) {
 			throw new Error(err.toString() + ': ' + (await err.responseBody))
 		}
-		// if (!body || body.toString().startsWith("No branch found")) {
-		// 	throw new Error(errorPrefix + `no branch found for FUNCTIONALTEST:${nodeName}`)
-		// }
-
-		// if (body.toUpperCase() === "RESOURCE NOT FOUND") {
-		// 	throw new Error(errorPrefix + `URL ${url} returned a "Resource Not Found" error.`)
-		// }
-
-		// const json = JSON.parse(body)
 
 		if (!json.branch) {
 			throw new Error(errorPrefix + `no branch found for ${botName}:${nodeName}`)
@@ -442,7 +436,6 @@ export abstract class FunctionalTest {
 		const key = this.getStreamPath(stream, depotName)
 		const client = allUserClients.get(key)
 
-		//console.log(`${this.testName}: Retrieving ${username}:${key}`)
 		if (!client) {
 			throw new Error(`${username} has no client for stream ${key}`)
 		}
@@ -460,7 +453,6 @@ export abstract class FunctionalTest {
 			throw new Error(`Client list for ${client.user} already contains entry for stream ${key}`)
 		}
 		clientMap.set(key, client)
-		//console.log(`${this.testName}: Adding ${client.user}:${key}:Client`)
 		this.clients.set(client.user, clientMap)
 		return client
 	}
@@ -481,6 +473,17 @@ export abstract class FunctionalTest {
 		await Promise.all(workspaceCreationPromises)
 	}
 
+	populateStreams(streams: Stream[]) {
+		return Promise.all(
+			streams
+				.filter(s => s.streamType !== 'mainline')
+				.map(s => this.p4.populate(
+					this.getStreamPath(s.name),
+					`Initial branch of files from ${s.parent}`
+				))
+		)
+	}
+
 	// for now, all added branches assumed to come from same source branch
 	async addTargetBranches(branches: RobomergeBranchSpec[], from: string, edgeProps?: EdgeProperties[]) {
 		const rootClient = await getRootDataClient(this.p4, 'RoboMergeData_BranchMaps')
@@ -535,7 +538,7 @@ export abstract class FunctionalTest {
 		}
 	}
 
-	async checkDescriptionContainsEdit(stream: string, requiredList?: string[], blacklist?: string[], depotName?: string) {
+	async checkDescriptionContainsEdit(stream: string, requiredList?: string[], unexpectedTerms?: string[], depotName?: string) {
 		this.info('Checking description of last commit to ' + stream)
 		this.getClient(stream, undefined, depotName).changes(1)
 			.then((changes: Change[]) => {
@@ -546,8 +549,8 @@ export abstract class FunctionalTest {
 						throw new Error(`Expected '${required}' to appear in description`)
 					}
 				}
-				if (blacklist) {
-					for (const bawal of blacklist) {
+				if (unexpectedTerms) {
+					for (const bawal of unexpectedTerms) {
 						if (description.indexOf(bawal.toLowerCase()) >= 0) {
 							this.error(description)
 							throw new Error(`Unexpected '${bawal}' in description`)
@@ -716,7 +719,12 @@ export abstract class FunctionalTest {
 		// wait for RoboMerge to process
 		await this.waitForRobomergeIdle()
 
-		return (await targetClient.changes(1, true))[0]!.change
+		const changes = await targetClient.changes(1, true)
+		if (changes.length === 0) {
+			this.error('Failed to create shelf for ' + target)
+			throw new Error('Failed to create shelf!')
+		}
+		return changes[0].change
 	}
 
 	storeNodesAndEdges() {
@@ -730,7 +738,7 @@ export abstract class FunctionalTest {
 		}
 	}
 
-	async isRobomergeIdle(dump?: boolean) {
+	async isRobomergeIdle(dump?: boolean): Promise<Map<string, number> | boolean | null>{
 
 		const latestP4CLs = new Map<string, Promise<number>>()
 		for (const client of this.clients.get('testuser1')!.values()) {
@@ -767,12 +775,6 @@ export abstract class FunctionalTest {
 			}
 
 			throw new Error('Branch state errors:\n' + branchStateErrors.join('\n'))
-			// if (true) {
-			// 	this.warn()
-			// 	console.log()
-			// }
-			// // this.warn('Assuming waiting for branchdefs to load')
-			// return false
 		}
 
 		for (const [node, bs] of unblockedBranchStates) {
@@ -804,8 +806,6 @@ export abstract class FunctionalTest {
 				}
 				return false
 			}
-
-			// console.log(`${node}: ${colors.magenta((await bs).getLastCL())} (${colors.cyan(await latestP4CLs.get(node)!)})`)
 		}
 
 		for (const [source, target] of this.edges) {
@@ -821,7 +821,7 @@ export abstract class FunctionalTest {
 				if (rmCL < p4CL) {
 					const gateClosed = edgeState.getGateClosedMessage()
 					if (gateClosed) {
-						this.verbose(gateClosed)
+						this.verbose(`${target} gate closed: '${gateClosed}'`)
 					}
 					else {
 						if (dump) {
@@ -835,14 +835,37 @@ export abstract class FunctionalTest {
 						return false
 					}
 				}
-				// console.log(`${source} -> ${target}: ${colors.magenta(edgeState.getLastCL())}, ${p4CL} - ${edgeState.isBlocked()}`)
 			}
 		}
-		return true
+
+		// when idle, produce a map of edge names to tick counts for edges neither blocked nor waiting for a gate
+		const ticksToWaitFor = new Map<string, number>([...unblockedBranchStates]
+			.filter(([_, node]) => !node.getEdges().every(e => e.getGateClosedMessage() || e.isBlocked()))
+			.map(([name, node]) => [name, node.getTickCount()]))
+
+		return ticksToWaitFor.size > 0 ? ticksToWaitFor : true
 	}
 
 	waitForRobomergeIdle(dump = false) {
-		return retryWithBackoff('Waiting for idle', async (last: boolean) => this.isRobomergeIdle(dump || last))
+		let tickCountsOnFirstIdle: Map<string, number> | null = null
+		return retryWithBackoff('Waiting for idle', async (last: boolean) => {
+			const tickCounts = await this.isRobomergeIdle(dump || last)
+			if (tickCounts === true) {
+				return true
+			}
+
+			if (!tickCounts) {
+				tickCountsOnFirstIdle = null
+			}
+			else if (!tickCountsOnFirstIdle) {
+				tickCountsOnFirstIdle = tickCounts
+			}
+			else {
+				return [...tickCounts].filter(([k, v]) => v <= (tickCountsOnFirstIdle!.get(k) || 0))
+  					.length === 0
+			}
+			return false
+		});
 	}
 }
 
diff --git a/Engine/Extras/RoboMerge/v3/functional_tests/framework/src/tests.ts b/Engine/Extras/RoboMerge/v3/functional_tests/framework/src/tests.ts
--- a/Engine/Extras/RoboMerge/v3/functional_tests/framework/src/tests.ts
+++ b/Engine/Extras/RoboMerge/v3/functional_tests/framework/src/tests.ts
@@ -10,11 +10,12 @@ import { ConfirmTextResolve } from './tests/confirm-text-resolve'
 import { ConfirmTextResolveBinaryStomp } from './tests/confirm-text-resolve-binary-stomp'
 import { CrossDepotStreamIntegration } from './tests/cross-depot-stream-integration'
 import { EdgeIndependence } from './tests/edge-independence'
-import { EdigrateMainRevToRelease } from './tests/edigrate-main-rev-to-release'
 import { ExclusiveCheckout } from './tests/exclusive-checkout'
 import { ExcludeAuthors } from './tests/exclude-authors'
 import { ExcludeAuthorsPerEdge } from './tests/exclude-authors-per-edge'
 import { ForwardCommands, ForwardCommands2 } from './tests/forward-commands'
+import { Ignore } from './tests/ignore'
+import { ImplicitCommands } from './tests/implicit-commands'
 import { IncognitoEdge } from './tests/incognito-edge'
 import { IncognitoTest } from './tests/incognito-test'
 import { IndirectTarget } from './tests/indirect-target'
@@ -30,6 +31,7 @@ import { RequestShelfIndirectTarget } from './tests/request-shelf-indirect-targe
 import { ResolveAfterSkip } from './tests/resolve-after-skip'
 import { RespectStreamPath } from './tests/respect-stream-path'
 import { SyntaxErrorOnUnknownBranch } from './tests/syntax-error-on-unknown-branch'
+import { StompForwardingCommands } from './tests/stomp-forwarding-commands'
 import { StompWithAdd } from './tests/stomp-with-add'
 import { TestChain } from './tests/test-chain'
 import { TestEdgeGate } from './tests/test-edge-gate'
@@ -145,10 +147,10 @@ async function go() {
 		new CrossDepotStreamIntegration(p4),
 		new EdgeIndependence(p4), // 5
 
-		new EdigrateMainRevToRelease(p4),
 		new ExclusiveCheckout(p4),
 		new IncognitoTest(p4),
 		new IndirectTarget(p4),
+		new ImplicitCommands(p4),
 		new MergeMainRevToMultipleRelease(p4), // 10
 
 		new MergeMainRevToRelease(p4),
@@ -191,7 +193,9 @@ async function go() {
 		new TestTerminal(p4),
 		new BlockIgnore(p4), // 40
 
-		new StompWithAdd(p4)
+		new StompWithAdd(p4),
+		new Ignore(p4),
+		new StompForwardingCommands(p4)
 	]
 
 	// const testToDebug = availableTests[30]
@@ -234,9 +238,7 @@ async function go() {
 	await Promise.all(tests.map(test => test.run()))
 
 	// wait for all tests after running, in case tests caused activity in other test streams (cross-bot, I'm looking at you)
-	for (const test of tests) {
-		await test.waitForRobomergeIdle()
-	}
+	await Promise.all(tests.map(test => test.waitForRobomergeIdle()))
 
 	console.log('Verifying tests')
 	let error: Error | null = null
diff --git a/Engine/Extras/RoboMerge/v3/functional_tests/framework/src/tests/block-ignore.ts b/Engine/Extras/RoboMerge/v3/functional_tests/framework/src/tests/block-ignore.ts
--- a/Engine/Extras/RoboMerge/v3/functional_tests/framework/src/tests/block-ignore.ts
+++ b/Engine/Extras/RoboMerge/v3/functional_tests/framework/src/tests/block-ignore.ts
@@ -1,5 +1,5 @@
 // Copyright Epic Games, Inc. All Rights Reserved.
-import { FunctionalTest, P4Client, P4Util, RobomergeBranchSpec, Stream } from '../framework'
+import { FunctionalTest, P4Util, RobomergeBranchSpec, Stream } from '../framework'
 
 export const streams: Stream[] = [
 	{name: 'Main', streamType: 'mainline'},
@@ -12,19 +12,22 @@ export class BlockIgnore extends FunctionalTest {
 		await this.p4.depot('stream', this.depotSpec())
 		await this.createStreamsAndWorkspaces(streams)
 
-		this.mainClient = this.getClient('Main')
-
-		await P4Util.addFileAndSubmit(this.mainClient, TEXT_FILENAME, 'Initial content')
+		await P4Util.addFileAndSubmit(this.getClient('Main'), TEXT_FILENAME, 'Initial content')
 
 		await this.p4.populate(this.getStreamPath('Release'), 'Initial branch of files from Main')
+		await this.getClient('Release').sync()
 	}
 
 	run() {
-		return P4Util.editFileAndSubmit(this.mainClient, TEXT_FILENAME, 'change', 'deadend')
+		return P4Util.editFileAndSubmit(this.getClient('Release'), TEXT_FILENAME, 'change', 'deadend')
 	}
 
 	verify() {
-		return Promise.all([this.ensureBlocked('Main'), this.ensureNotBlocked('Release')])
+		return Promise.all([
+			this.ensureNotBlocked('Main'),
+			this.ensureBlocked('Release'),
+			this.checkHeadRevision('Main', TEXT_FILENAME, 1)
+		])
 	}
 
 	allowSyntaxErrors() {
@@ -32,12 +35,9 @@ export class BlockIgnore extends FunctionalTest {
 	}
 
 	getBranches() {
-		const main = this.branch('Main', [])
-		main.disallowDeadend = true
-		return [
-			main,
-			this.branch('Release', ['Main'])
-		]
+		const release = this.branch('Release', ['Main'])
+		release.disallowDeadend = true
+		return [this.branch('Main', []), release]
 	}
 
 	private branch(stream: string, to: string[]): RobomergeBranchSpec {
@@ -49,6 +49,4 @@ export class BlockIgnore extends FunctionalTest {
 			forceAll: true
 		}
 	}
-
-	private mainClient: P4Client
 }
diff --git a/Engine/Extras/RoboMerge/v3/functional_tests/framework/src/tests/ignore.ts b/Engine/Extras/RoboMerge/v3/functional_tests/framework/src/tests/ignore.ts
new file mode 100644
--- /dev/null
+++ b/Engine/Extras/RoboMerge/v3/functional_tests/framework/src/tests/ignore.ts
@@ -0,0 +1,53 @@
+// Copyright Epic Games, Inc. All Rights Reserved.
+import { FunctionalTest, P4Util, RobomergeBranchSpec, Stream } from '../framework'
+
+export const streams: Stream[] = [
+	{name: 'Main', streamType: 'mainline'},
+	{name: 'Release', streamType: 'release', parent: 'Main'}
+]
+
+const TEXT_FILENAME = 'test.txt'
+export class Ignore extends FunctionalTest {
+	async setup() {
+		await this.p4.depot('stream', this.depotSpec())
+		await this.createStreamsAndWorkspaces(streams)
+
+		await P4Util.addFileAndSubmit(this.getClient('Main'), TEXT_FILENAME, 'Initial content')
+
+		await this.p4.populate(this.getStreamPath('Release'), 'Initial branch of files from Main')
+		await this.getClient('Release').sync()
+	}
+
+	run() {
+		return P4Util.editFileAndSubmit(this.getClient('Release'), TEXT_FILENAME, 'change', 'deadend')
+	}
+
+	verify() {
+		return Promise.all([
+			this.ensureNotBlocked('Main'),
+			this.ensureNotBlocked('Release'),
+			this.checkHeadRevision('Main', TEXT_FILENAME, 1)
+		])
+	}
+
+	allowSyntaxErrors() {
+		return true
+	}
+
+	getBranches() {
+		return [
+			this.branch('Main', []),
+			this.branch('Release', ['Main'])
+		]
+	}
+
+	private branch(stream: string, to: string[]): RobomergeBranchSpec {
+		return {
+			streamDepot: this.constructor.name,
+			name: this.fullBranchName(stream),
+			streamName: stream,
+			flowsTo: to.map(str => this.fullBranchName(str)),
+			forceAll: true
+		}
+	}
+}
diff --git a/Engine/Extras/RoboMerge/v3/functional_tests/framework/src/tests/implicit-commands.ts b/Engine/Extras/RoboMerge/v3/functional_tests/framework/src/tests/implicit-commands.ts
new file mode 100644
--- /dev/null
+++ b/Engine/Extras/RoboMerge/v3/functional_tests/framework/src/tests/implicit-commands.ts
@@ -0,0 +1,68 @@
+// Copyright Epic Games, Inc. All Rights Reserved.
+
+import { EdgeProperties, FunctionalTest, P4Util, Stream } from '../framework'
+
+const streams: Stream[] = [
+	{name: 'Main', streamType: 'mainline'},
+	{name: 'Producer1', streamType: 'release', parent: 'Main'},
+	{name: 'Producer2', streamType: 'release', parent: 'Main'},
+	{name: 'Consumer', streamType: 'development', parent: 'Main'}
+]
+
+
+export class ImplicitCommands extends FunctionalTest {
+	async setup() {
+		await this.p4.depot('stream', this.depotSpec())
+		await this.createStreamsAndWorkspaces(streams)
+
+		await P4Util.addFileAndSubmit(this.getClient('Main'), 'test1.txt', 'Initial content')
+		await P4Util.addFileAndSubmit(this.getClient('Main'), 'test2.txt', 'Initial content')
+
+		const desc = 'Initial branch of files from Main'
+		await Promise.all([
+			this.p4.populate(this.getStreamPath('Producer1'), desc),
+			this.p4.populate(this.getStreamPath('Producer2'), desc),
+			this.p4.populate(this.getStreamPath('Consumer'), desc)
+		])
+
+		await Promise.all(['Producer1', 'Producer2'].map(
+			branch => this.getClient(branch).sync()
+		))
+	}
+
+	run() {
+		return Promise.all([
+			P4Util.editFileAndSubmit(this.getClient('Producer1'), 'test1.txt', 'new content'),
+			P4Util.editFileAndSubmit(this.getClient('Producer2'), 'test2.txt', 'new content')
+		])
+	}
+
+	verify() {
+		return Promise.all([
+			this.checkHeadRevision('Main', 'test1.txt', 2), // main should have both changes
+			this.checkHeadRevision('Main', 'test2.txt', 2),
+			this.checkHeadRevision('Producer1', 'test1.txt', 2),
+			this.checkHeadRevision('Producer2', 'test2.txt', 2),
+			this.checkHeadRevision('Consumer', 'test1.txt', 1), // consumer should only have test2 change
+			this.checkHeadRevision('Consumer', 'test2.txt', 2)
+		])
+	}
+
+	getBranches() {
+		return [
+			this.makeForceAllBranchDef('Main', ['Consumer']),
+			this.makeForceAllBranchDef('Producer1', ['Main']),
+			this.makeForceAllBranchDef('Producer2', ['Main']),
+			this.makeForceAllBranchDef('Consumer', [])
+		]
+	}
+
+	getEdges(): EdgeProperties[] {
+		return [
+		  { from: this.fullBranchName('Producer1'), to: this.fullBranchName('Main')
+		  , implicitCommands: ['#robomerge -' + this.fullBranchName('Consumer')]
+		  }
+		]
+	}
+
+}
diff --git a/Engine/Extras/RoboMerge/v3/functional_tests/framework/src/tests/stomp-forwarding-commands.ts b/Engine/Extras/RoboMerge/v3/functional_tests/framework/src/tests/stomp-forwarding-commands.ts
new file mode 100644
--- /dev/null
+++ b/Engine/Extras/RoboMerge/v3/functional_tests/framework/src/tests/stomp-forwarding-commands.ts
@@ -0,0 +1,59 @@
+import { FunctionalTest, P4Client, P4Util, Stream } from '../framework'
+
+const streams: Stream[] = [
+	{name: 'Main', streamType: 'mainline'},
+	{name: 'DevChild', streamType: 'development', parent: 'Main'},
+	{name: 'DevGrandchild', streamType: 'development', parent: 'DevChild'},
+]
+
+export class StompForwardingCommands extends FunctionalTest {
+	async setup() {
+		await this.p4.depot('stream', this.depotSpec())
+		await this.createStreamsAndWorkspaces(streams)
+
+		this.mainClient = this.getClient('Main')
+		await P4Util.addFileAndSubmit(this.mainClient, 'test.uasset', 'Initial content', true)
+
+
+		await this.populateStreams(streams)
+
+		const childClient = this.getClient('DevChild')
+		await childClient.sync()
+
+		// set up conflict
+		await P4Util.editFileAndSubmit(childClient, 'test.uasset', 'Changed content')
+
+	}
+
+	async run() {
+		await P4Util.editFileAndSubmit(this.mainClient, 'test.uasset', 'Conflicting change',
+			'-' + this.fullBranchName('DevGrandchild'))
+
+		await this.waitForRobomergeIdle()
+
+		await this.ensureBlocked('Main', 'DevChild')
+		await this.verifyAndPerformStomp('Main', 'DevChild')
+	}
+
+	verify() {
+		// make sure skip command was included in stomped file
+		return Promise.all([
+			this.checkHeadRevision('DevChild', 'test.uasset', 3),
+			this.checkHeadRevision('DevGrandchild', 'test.uasset', 2),
+		])
+	}
+
+	getBranches() {
+		const childSpec = this.makeForceAllBranchDef('DevChild', ['DevGrandchild'])
+		childSpec.initialCL = 1
+
+		return [
+			this.makeForceAllBranchDef('Main', ['DevChild']),
+			childSpec,
+			this.makeForceAllBranchDef('DevGrandchild', []),
+		]
+	}
+
+	mainClient: P4Client
+}
+
diff --git a/Engine/Extras/RoboMerge/v3/package.json b/Engine/Extras/RoboMerge/v3/package.json
--- a/Engine/Extras/RoboMerge/v3/package.json
+++ b/Engine/Extras/RoboMerge/v3/package.json
@@ -31,14 +31,16 @@
     "marked": "^0.7.0",
     "mustache": "^3.2.1",
     "node-gzip": "^1.1.2",
-    "nodemailer": "^4.6.3",
+    "nodemailer": "^6.7.0",
+    "nomnom": "^1.5.3",
     "request": "^2.88.0",
-    "sanitize-html": "^1.20.0",
+    "sanitize-html": "^2.5.3",
     "semver": "^6.1.0",
     "source-map-support": "^0.5.12",
     "winston": "^3.2.1"
   },
   "devDependencies": {
+    "@types/intl": "^1.2.0",
     "@types/swagger-ui-express": "^3.0.1"
   }
 }
diff --git a/Engine/Extras/RoboMerge/v3/public/allbots.html b/Engine/Extras/RoboMerge/v3/public/allbots.html
--- a/Engine/Extras/RoboMerge/v3/public/allbots.html
+++ b/Engine/Extras/RoboMerge/v3/public/allbots.html
@@ -19,8 +19,11 @@
 			type: 'get',
 			dataType: 'json',
 			success: data => {
-				const bots = location.search.match(/bots=(.*)/)
-				$(() => $('#graph').append(showFlowGraph(data.branches, bots ? bots[1].split(',') : [])))
+				const bots = location.search.match(/bots=([^&]*)/);
+				$(() => $('#graph').append(showFlowGraph(data.branches, {
+					botNames: bots ? bots[1].split(',') : [],
+					...parseOptions(location.search)
+				})));
 			},
 		// error: function(xhr, error, status) {
 		// 	robomergeUser = undefined;
@@ -32,6 +35,23 @@
 </head>
 <body>
 	<header><img src="/img/logo.png"></header>
-	<div id="graph">
+	<div class="container">
+		<div id="graph">
+		</div>
+
+		<div style="transform: scale(0.7);">
+			<h4>Options</h4>
+			<table class="table">
+				<thead><tr><th>Setting</th><th>Effect</th></tr></thead>
+				<tbody>
+					<tr><td><tt>bots=</tt></td><td>Comma separated list of bots to show</td></tr>
+					<tr><td><tt>hideDisconnected</tt></td><td>Hide nodes with no edges</td></tr>
+					<tr><td><tt>noGroups</tt></td><td>Relevant if showing multiple bots: suppress boxes grouping each bot</td></tr>
+					<tr><td><tt>showOnlyForced</tt></td><td>Omit any "on demand" (dotted gray) edges</td></tr>
+				</tbody>
+			</table>
+
+			<p>e.g. <tt>http://robomerge/allbots/12345?bots=test1,test2&noGroups&showOnlyForced</tt></p>
+		</div>
 	</div>
 </body>
diff --git a/Engine/Extras/RoboMerge/v3/public/createshelf.html b/Engine/Extras/RoboMerge/v3/public/createshelf.html
--- a/Engine/Extras/RoboMerge/v3/public/createshelf.html
+++ b/Engine/Extras/RoboMerge/v3/public/createshelf.html
@@ -29,7 +29,7 @@
     <div id="result"></div>
     <div id="nextSteps" class="initiallyHidden centered-text">
         <h3>What next?</h3>
-        <p>Robomerge will run the merge for the changelist again and create any shelfed conflicts in the workspace you chose. You should receive an email when this process is complete.</p>
+        <p>Robomerge will run the merge for the changelist again and create any shelved conflicts in the workspace you chose. You should receive an email when this process is complete.</p>
         <p>If you do not receive an email, please <strong><a href="/contact">Contact Us</a></strong> for Robomerge help.</p>
     </div>
     <div id="createStreamWorkspaceTutorial" class="initiallyHidden centered-text">
diff --git a/Engine/Extras/RoboMerge/v3/public/css/main.css b/Engine/Extras/RoboMerge/v3/public/css/main.css
--- a/Engine/Extras/RoboMerge/v3/public/css/main.css
+++ b/Engine/Extras/RoboMerge/v3/public/css/main.css
@@ -333,4 +333,8 @@ label.workspaceInexactMatch {
 
 .big-help {
 	margin: 5px;
+}
+
+.warning {
+	color: orange;
 }
\ No newline at end of file
diff --git a/Engine/Extras/RoboMerge/v3/public/css/table.css b/Engine/Extras/RoboMerge/v3/public/css/table.css
--- a/Engine/Extras/RoboMerge/v3/public/css/table.css
+++ b/Engine/Extras/RoboMerge/v3/public/css/table.css
@@ -245,6 +245,11 @@
     font-size: 0.75em;
   }
 
+  .statuscell .status-msg {
+    font-style: italic;
+    font-size: .9em;
+  }
+
   .lastchangecell .catchup {
     font-size: 0.75em;
     line-height: 1.1;
diff --git a/Engine/Extras/RoboMerge/v3/public/js/boilerplate.js b/Engine/Extras/RoboMerge/v3/public/js/boilerplate.js
--- a/Engine/Extras/RoboMerge/v3/public/js/boilerplate.js
+++ b/Engine/Extras/RoboMerge/v3/public/js/boilerplate.js
@@ -935,7 +935,14 @@ function renderNameCell_Common(data, botname) {
 		})
 	}
 
-	if (!data.is_available) {
+	if (data.is_available) {
+		if (data.windowBypass) {
+			$('<div>').addClass('status-msg warning')
+				.appendTo(topInternalDiv)
+				.text('Bypassing window!');
+		}
+	}
+	else {
 		const unpauseTime = data.blockage && data.blockage.endsAt ? new Date(data.blockage.endsAt) : 'never';
 		let msg = '';
 
@@ -953,15 +960,6 @@ function renderNameCell_Common(data, botname) {
 		if (data.is_paused || data.is_blocked) {
 			msg += `Will unpause and retry on: ${unpauseTime}.<br />`;
 		}
-		else if (data.gateClosedMessage) {
-			msg = data.gateClosedMessage
-			if (data.nextWindowOpenTime) {
-				const timestamp = Date.parse(data.nextWindowOpenTime)
-				if (!isNaN(timestamp)) {
-					msg += ` (opens: ${(new Date(timestamp)).toLocaleString()})`
-				}
-			}
-		}
 
 		if (msg) {
 			$('<div class="status-msg">')
@@ -1133,7 +1131,7 @@ function preRenderStatusCell_Node(nodeData) {
 }
 function preRenderStatusCell_Edge(nodeData, edgeData) {
 	// Create holder for the stylized status text
-	let statusCell = $('<div class="statuscell">')
+	let statusCell = $('<div>').addClass('statuscell')
 
 	if (edgeData.is_blocked) {
 		$('<div class="blockageinfo">')
@@ -1152,6 +1150,16 @@ function preRenderStatusCell_Edge(nodeData, edgeData) {
 				.css('color', 'ORANGE')
 				.prependTo(statusCell)
 		}
+		else if (edgeData.gateClosedMessage) {
+			let msg = edgeData.gateClosedMessage
+			if (edgeData.nextWindowOpenTime) {
+				const timestamp = Date.parse(edgeData.nextWindowOpenTime)
+				if (!isNaN(timestamp)) {
+					msg += ` (opens: ${(new Date(timestamp)).toLocaleString()})`
+				}
+			}
+			$('<div>').addClass('status-msg').text(msg).appendTo(statusCell)
+		}
 	}
 
 	return statusCell
@@ -1331,12 +1339,12 @@ function renderActionsCell_Common(actionCell, data, operationFunction, operation
 				'Take ownership of this blockage'
 			)
 		} 
-		// Unackowledge
+		// Unacknowledge
 		else if (robomergeUser && data.blockage.acknowledger === robomergeUser.userName) {
 			specificActionButton = createActionButton(
 				[$('<i class="fas fa-eject" style="color:red">'), '&nbsp;', "Unacknowledge"],
 				unackFunc,
-				'Relinguish control over this conflict'
+				'Back-pedal: someone else should fix this'
 			)
 		}
 		// Take ownership
@@ -1374,7 +1382,7 @@ function renderActionsCell_Common(actionCell, data, operationFunction, operation
 			}, `Retry merge of ${conflict.cl} and create a shelf in a specified P4 workspace`)
 			
 			// Stomp
-			const stompRequest = `/op/stomp?` + queryParams
+			const stompRequest = `/op/stomp?` + queryParams + location.hash
 			const stompOption = createActionOption('Stomp Changes using ' + conflict.cl + toTargetText, function() {
 				window.location.href = stompRequest;
 			}, `Use ${conflict.cl} to stomp binary changes in ${conflict.target}`);
@@ -1569,6 +1577,23 @@ function renderActionsCell_Edge(actionCell, nodeData, edgeData, conflict=null) {
 			)
 		}
 	}
+
+	if (edgeData.nextWindowOpenTime) {
+		const sense = edgeData.windowBypass ? 'false' : 'true'
+
+		const text = edgeData.windowBypass
+			? 'Re-enable gate window'
+			: 'Bypass gate window';
+
+		const tooltip = edgeData.windowBypass
+			? 'Integrations will be delayed until window is open'
+			: 'Integrations will proceed, disregarding specified window';
+
+		const bypassRequest = () => 
+			edgeAPIOp(nodeData.bot, getAPIName(nodeData), edgeData.target, '/bypassgatewindow?sense=' + sense, () => {}, () => {});
+
+		insertFunction(createActionOption(text, bypassRequest, tooltip));
+	}
 }
 
 // Helper function to create last change cell
diff --git a/Engine/Extras/RoboMerge/v3/public/js/flow.js b/Engine/Extras/RoboMerge/v3/public/js/flow.js
--- a/Engine/Extras/RoboMerge/v3/public/js/flow.js
+++ b/Engine/Extras/RoboMerge/v3/public/js/flow.js
@@ -1,4 +1,21 @@
 // Copyright Epic Games, Inc. All Rights Reserved.
+class Info {
+    constructor(id, tooltip, name, sources) {
+        this.id = id;
+        this.tooltip = tooltip;
+        this.name = name;
+        this.sources = sources;
+        this.forcedDests = new Set();
+        this.defaultDests = new Set();
+        this.blockAssetDests = new Set();
+        this.importance = 0;
+    }
+    update(branch, getId) {
+        this.forcedDests = new Set([...this.forcedDests, ...branch.def.forceFlowTo.map(getId)]);
+        this.defaultDests = new Set([...this.defaultDests, ...branch.def.defaultFlow.map(getId)]);
+        this.blockAssetDests = new Set([...this.blockAssetDests, ...branch.def.blockAssetTargets.map(getId)]);
+    }
+}
 const EDGE_STYLES = {
     roboshelf: [['color', 'purple'], ['arrowhead', 'diamond']],
     forced: [],
@@ -21,127 +38,73 @@ function renderGraph(src) {
 function decoratedAlias(bot, alias) {
     return bot + ':' + alias;
 }
-function showFlowGraph(data, botNameParam) {
-    let whitelist = null;
-    let botName;
-    if (Array.isArray(botNameParam)) {
-        botName = null;
-        if (botNameParam.length > 0) {
-            whitelist = botNameParam.map(s => s.toUpperCase());
-        }
+function aliasForBranch(branch) {
+    return decoratedAlias(branch.bot, branch.def.upperName);
+}
+function parseOptions(search) {
+    return {
+        showOnlyForced: !!search.match(/showOnlyForced/i),
+        noGroups: !!search.match(/noGroups/i),
+        hideDisconnected: !!search.match(/hideDisconnected/i)
+    };
+}
+class Graph {
+    constructor(allBranches, options) {
+        this.allBranches = allBranches;
+        this.options = options;
+        this.nodeLabels = new Map();
+        this.links = [];
+        this.aliases = new Map();
+        this.allInfos = [];
+        this.linkOutward = new Set();
+        this.linkInward = new Set();
+        this.connectedNodes = new Set();
+        this.branchList = [];
     }
-    else {
-        botName = botNameParam;
+    singleBot(botName) {
+        this.branchList = this.allBranches.filter(b => b.bot === botName);
+        this.addNodesAndEdges(false);
+        return this.makeGraph(botName, botName + 'integration paths');
     }
-    const addLinesForBranchGraph = (outLines, bot, branchList) => {
-        const aliases = new Map();
-        const nodeLabels = new Map();
-        function getIdForBot(bot) {
-            return (alias) => {
-                const info = aliases.get(decoratedAlias(bot, alias));
-                return info ? info.id : alias;
-            };
-        }
-        // if showing all bots, build up map from stream to nodes
-        if (!botName) {
-            const streamMap = new Map();
-            for (let branchStatus of branchList) {
-                const [color, streams] = setDefault(streamMap, branchStatus.def.rootPath, [null, []]);
-                streams.push(branchStatus);
-                const nodeColor = branchStatus.def.config.graphNodeColor;
-                if (!color && nodeColor) {
-                    streamMap.set(branchStatus.def.rootPath, [nodeColor, streams]);
-                }
-            }
-            // create a shared info for each set of nodes monitoring the same stream
-            for (const [k, v] of streamMap) {
-                const [color, streams] = v;
-                if (streams.length > 1) {
-                    const info = {
-                        id: k.replace(/[^\w]+/g, '_'), tooltip: 'tooltip todo', name: k
-                    };
-                    if (color) {
-                        info.graphNodeColor = color;
-                    }
-                    for (const branchStatus of streams) {
-                        for (const alias of branchStatus.def.aliases) {
-                            aliases.set(decoratedAlias(branchStatus.bot, alias), info);
-                        }
-                    }
-                    setDefault(nodeLabels, 'nogroup', []).push(info);
-                }
-            }
+    allBots() {
+        this.branchList = this.options.botsToShow
+            ? this.allBranches.filter(b => this.options.botsToShow.indexOf(b.bot) >= 0)
+            : this.allBranches;
+        this.findSharedNodes();
+        this.addNodesAndEdges(!this.options.noGroups);
+        return this.makeGraph('All bots', 'Flow including shared bots');
+    }
+    addNodesAndEdges(grouped) {
+        for (const branchStatus of this.branchList) {
+            this.addBranch(branchStatus, grouped ? branchStatus.bot : null);
         }
-        // build map of all aliases and list of node information
-        for (let branchStatus of branchList) {
-            const branch = branchStatus.def;
-            if (aliases.has(decoratedAlias(branchStatus.bot, branch.upperName)))
-                continue;
-            let tooltip = branch.rootPath;
-            if (branch.aliases && branch.aliases.length !== 0) {
-                tooltip += ` (${branch.aliases.join(', ')})`;
-            }
-            const info = {
-                id: `_${branchStatus.bot}_${branch.upperName.replace(/[^\w]+/g, '_')}`,
-                tooltip, name: branch.name
-            };
-            if (branch.config.graphNodeColor) {
-                info.graphNodeColor = branch.config.graphNodeColor;
-            }
-            // note: branch.upperName is always in branch.aliases
-            for (const alias of branch.aliases) {
-                aliases.set(decoratedAlias(branchStatus.bot, alias), info);
+        for (const info of this.allInfos) {
+            for (const sourceBranch of info.sources) {
+                info.update(sourceBranch, this.getIdForBot(sourceBranch.bot));
             }
-            setDefault(nodeLabels, botName ? 'nogroup' : branchStatus.bot, []).push(info);
         }
-        const branchesById = new Map();
-        for (const branchStatus of branchList) {
-            const branch = branchStatus.def;
-            const getId = getIdForBot(branchStatus.bot);
-            const branchId = getId(branch.upperName);
-            branchesById.set(branchId, {
-                id: branchId,
-                forcedDests: new Set(branch.forceFlowTo.map(getId)),
-                defaultDests: new Set(branch.defaultFlow.map(getId)),
-                blockAssetDests: new Set(branch.blockAssetTargets ? branch.blockAssetTargets.map(getId) : [])
-            });
-        }
-        const linkOutward = new Set(), linkInward = new Set();
-        const nodeImportance = new Map();
-        // special case branches named Main
-        nodeImportance.set(getIdForBot(bot)('MAIN'), 10);
-        const links = [];
-        for (const branchStatus of branchList) {
-            const getId = getIdForBot(branchStatus.bot);
-            const src = getId(branchStatus.def.upperName);
-            const branch = branchesById.get(src);
-            for (let flow of branchStatus.def.flowsTo) {
-                const dst = getId(flow);
-                nodeImportance.set(src, (nodeImportance.get(src) || 0) + 1);
-                nodeImportance.set(dst, (nodeImportance.get(dst) || 0) + 1);
-                const isForced = branch.forcedDests.has(dst);
-                const edgeStyle = branchStatus.def.convertIntegratesToEdits ? 'roboshelf' :
-                    isForced ? 'forced' :
-                        branch.defaultDests.has(dst) ? 'defaultFlow' :
-                            branch.blockAssetDests.has(dst) ? 'blockAssets' : 'onRequest';
-                const styles = [...EDGE_STYLES[edgeStyle]];
-                const link = { src, dst, styles };
-                links.push(link);
-                if (isForced) {
-                    linkOutward.add(src + dst);
-                }
-                else if (edgeStyle === 'blockAssets' || edgeStyle === 'onRequest') {
-                    linkInward.add(dst + src);
-                }
+        for (const branchStatus of this.branchList) {
+            for (const flow of branchStatus.def.flowsTo) {
+                this.addEdge(branchStatus, flow);
             }
         }
+    }
+    makeGraph(title, tooltip) {
+        const lines = [
+            'digraph robomerge {',
+            'fontname="sans-serif"; labelloc=top; fontsize=16;',
+            'edge [penwidth=2]; nodesep=.7; ranksep=1.2;',
+            `label = "${title}";`,
+            `tooltip="${tooltip}";`,
+            `node [shape=box, style=filled, fontname="sans-serif", fillcolor=moccasin];`,
+        ];
         const nodeGroups = [];
-        for (const [groupName, v] of nodeLabels) {
+        for (const [groupName, v] of this.nodeLabels) {
             const nodeInfo = [];
             nodeGroups.push([groupName, nodeInfo]);
             for (const info of v) {
-                const { id, tooltip, name } = info;
-                let factor = (Math.min(nodeImportance.get(id) || 0, 10) - 1) / 9;
+                const { id, tooltip, name, importance } = info;
+                let factor = (Math.min(importance, 10) - 1) / 9;
                 nodeInfo.push({
                     importance: factor,
                     marginX: (.2 * (1 - factor) + .4 * factor).toPrecision(1),
@@ -151,13 +114,16 @@ function showFlowGraph(data, botNameParam) {
                 });
             }
         }
-        for (let [groupName, nodeInfo] of nodeGroups) {
+        for (const [groupName, nodeInfo] of nodeGroups) {
             if (groupName !== 'nogroup') {
-                outLines.push(`subgraph cluster_${groupName} {
+                lines.push(`subgraph cluster_${groupName} {
 	label="${groupName}";
 `);
             }
             for (const info of nodeInfo) {
+                if (this.options.hideDisconnected && !this.connectedNodes.has(info.info.id)) {
+                    continue;
+                }
                 const attrs = [
                     ['label', `"${info.info.name}"`],
                     ['tooltip', `"${info.info.tooltip}"`],
@@ -171,100 +137,144 @@ function showFlowGraph(data, botNameParam) {
                     attrs.push(['fillcolor', `"${info.info.graphNodeColor}"`]);
                 }
                 const attrStrs = attrs.map(([key, value]) => `${key}=${value}`);
-                outLines.push(`${info.info.id} [${attrStrs.join(', ')}];`);
+                lines.push(`${info.info.id} [${attrStrs.join(', ')}];`);
             }
             if (groupName !== 'nogroup') {
-                outLines.push('}');
+                lines.push('}');
             }
         }
-        for (const link of links) {
+        for (const link of this.links) {
             // when there's forced/unforced pair, only the former is a constrait
             // (this makes flow )
-            const combo = link.dst + link.src;
-            if (combo && linkOutward.has(combo) && linkInward.has(combo)) {
+            const combo = link.dst + '->' + link.src;
+            if (combo && this.linkOutward.has(combo) && this.linkInward.has(combo)) {
                 link.styles.push(['constraint', 'false']);
             }
             const styleStrs = link.styles.map(([key, value]) => `${key}=${value}`);
             const suffix = styleStrs.length === 0 ? '' : ` [${styleStrs.join(', ')}]`;
-            outLines.push(`${link.src} -> ${link.dst}${suffix};`);
+            lines.push(`${link.src} -> ${link.dst}${suffix};`);
         }
-    };
-    let bots = new Map();
-    for (let branch of data) {
-        let bot;
-        if (botName) {
-            if (branch.bot !== botName)
-                continue;
-            bot = branch.bot;
+        lines.push('}');
+        return lines;
+    }
+    addBranch(branchStatus, group = null) {
+        const branch = branchStatus.def;
+        if (this.aliases.has(decoratedAlias(branchStatus.bot, branch.upperName)))
+            return;
+        let tooltip = branch.rootPath;
+        if (branch.aliases && branch.aliases.length !== 0) {
+            tooltip += ` (${branch.aliases.join(', ')})`;
         }
-        else {
-            if (whitelist && whitelist.indexOf(branch.bot) < 0)
-                continue;
-            bot = 'all';
+        const info = new Info(`_${branchStatus.bot}_${branch.upperName.replace(/[^\w]+/g, '_')}`, tooltip, branch.name, [branchStatus]);
+        this.allInfos.push(info);
+        if (branch.config.graphNodeColor) {
+            info.graphNodeColor = branch.config.graphNodeColor;
         }
-        let branchList = bots.get(bot);
-        if (!branchList) {
-            branchList = [];
-            bots.set(bot, branchList);
+        if (branch.upperName === 'MAIN') {
+            // special case branches named Main
+            info.importance = 10;
         }
-        branchList.push(branch);
+        // note: branch.upperName is always in branch.aliases
+        for (const alias of branch.aliases) {
+            this.aliases.set(decoratedAlias(branchStatus.bot, alias), info);
+        }
+        setDefault(this.nodeLabels, group || 'nogroup', []).push(info);
     }
-    let graphSrcIter = (function* () {
-        let preface = [
-            'digraph robomerge {',
-            'fontname="sans-serif"; labelloc=top; fontsize=16;',
-            'edge [penwidth=2]; nodesep=.7; ranksep=1.2;',
-        ];
-        for (let [bot, branchList] of bots) {
-            let lines = preface.slice();
-            lines.push(`label = "${bot}";`);
-            lines.push(`tooltip="${bot} integration paths";`);
-            lines.push(`node [shape=box, style=filled, fontname="sans-serif", fillcolor=moccasin];`);
-            addLinesForBranchGraph(lines, bot, branchList);
-            lines.push('}');
-            yield [bot, lines.join('\n')];
+    // dstName 
+    addEdge(srcBranchStatus, dst) {
+        const src = aliasForBranch(srcBranchStatus);
+        const srcInfo = this.aliases.get(src);
+        const dstInfo = this.aliases.get(decoratedAlias(srcBranchStatus.bot, dst));
+        ++srcInfo.importance;
+        ++dstInfo.importance;
+        const isForced = srcInfo.forcedDests.has(dstInfo.id);
+        if (!isForced && this.options.showOnlyForced) {
+            return;
         }
-    })();
-    let graphContainer = $('<div class="clearfix">');
-    let flowGraph = $('<div class="flow-graph" style="display:inline-block;">').appendTo(graphContainer);
-    let first = true;
-    flowGraph.append($('<div>').css('text-align', 'center').text("Building graph..."));
-    let present = (bot, graphSvg) => {
-        if (graphSvg) {
-            $('#graph-key-template')
-                .clone()
-                .removeAttr('id')
-                .css('display', 'inline-block')
-                .appendTo(graphContainer);
-            $('#graph-loading-text').hide();
-            let span = $('<span>').css('margin', 'auto').html(graphSvg);
-            let svgEl = $('svg', span).addClass('branch-graph').removeAttr('width');
-            // scale graph to 70% of default size
-            let height = Math.round(parseInt(svgEl.attr('height')) * .7);
-            svgEl.attr('height', height + 'px').css('vertical-align', 'top');
-            let inserted = false;
-            for (let existing of $('svg', flowGraph)) {
-                let el = $(existing);
-                if (parseInt(el.attr('height')) < height) {
-                    span.insertBefore(el.parent());
-                    inserted = true;
-                    break;
-                }
-            }
-            if (first) {
-                first = false;
-                flowGraph.empty();
+        this.connectedNodes.add(srcInfo.id);
+        this.connectedNodes.add(dstInfo.id);
+        const edgeStyle = srcBranchStatus.def.convertIntegratesToEdits ? 'roboshelf' :
+            isForced ? 'forced' :
+                srcInfo.defaultDests.has(dstInfo.id) ? 'defaultFlow' :
+                    srcInfo.blockAssetDests.has(dstInfo.id) ? 'blockAssets' : 'onRequest';
+        const styles = [...EDGE_STYLES[edgeStyle]];
+        const link = { src: srcInfo.id, dst: dstInfo.id, styles };
+        this.links.push(link);
+        if (isForced) {
+            this.linkOutward.add(link.src + '->' + link.dst);
+        }
+        else if (edgeStyle === 'blockAssets' || edgeStyle === 'onRequest') {
+            this.linkInward.add(link.dst + '->' + link.src);
+        }
+    }
+    findSharedNodes() {
+        const streamMap = new Map();
+        for (const branchStatus of this.branchList) {
+            const [color, streams] = setDefault(streamMap, branchStatus.def.rootPath, [null, []]);
+            streams.push(branchStatus);
+            const nodeColor = branchStatus.def.config.graphNodeColor;
+            if (!color && nodeColor) {
+                streamMap.set(branchStatus.def.rootPath, [nodeColor, streams]);
             }
-            if (!inserted)
-                flowGraph.append(span);
         }
-        const { value, done } = graphSrcIter.next();
-        if (!done) {
-            const [bot, src] = value;
-            console.log(src);
-            renderGraph(src).then(graphSvg => present(bot, graphSvg));
+        // create a shared info for each set of nodes monitoring the same stream
+        for (const [k, v] of streamMap) {
+            const [color, streams] = v;
+            if (streams.length > 1) {
+                const info = new Info(k.replace(/[^\w]+/g, '_'), streams.map(s => s.bot).join(', '), k, streams);
+                this.allInfos.push(info);
+                if (color) {
+                    info.graphNodeColor = color;
+                }
+                for (const branchStatus of streams) {
+                    for (const alias of branchStatus.def.aliases) {
+                        this.aliases.set(decoratedAlias(branchStatus.bot, alias), info);
+                    }
+                }
+                setDefault(this.nodeLabels, 'nogroup', []).push(info);
+            }
         }
-    };
-    present();
+    }
+    getIdForBot(bot) {
+        return (alias) => {
+            const info = this.aliases.get(decoratedAlias(bot, alias));
+            return info ? info.id : alias;
+        };
+    }
+}
+function makeGraph(data, args) {
+    let botsToShow = null;
+    const options = args;
+    const singleBotName = args.singleBotName;
+    if (singleBotName) {
+        return (new Graph(data, options)).singleBot(singleBotName);
+    }
+    const botNames = args.botNames;
+    if (botNames.length > 0) {
+        options.botsToShow = botNames.map(s => s.toUpperCase());
+    }
+    return (new Graph(data, options)).allBots();
+}
+function showFlowGraph(data, args) {
+    const lines = makeGraph(data, args);
+    const graphContainer = $('<div class="clearfix">');
+    const flowGraph = $('<div class="flow-graph" style="display:inline-block;">').appendTo(graphContainer);
+    flowGraph.append($('<div>').css('text-align', 'center').text("Building graph..."));
+    renderGraph(lines.join('\n'))
+        .then(svg => {
+        $('#graph-key-template')
+            .clone()
+            .removeAttr('id')
+            .css('display', 'inline-block')
+            .appendTo(graphContainer);
+        $('#graph-loading-text').hide();
+        const span = $('<span>').css('margin', 'auto').html(svg);
+        const svgEl = $('svg', span).addClass('branch-graph').removeAttr('width');
+        // scale graph to 70% of default size
+        const height = Math.round(parseInt(svgEl.attr('height')) * .7);
+        svgEl.attr('height', height + 'px').css('vertical-align', 'top');
+        flowGraph.empty();
+        flowGraph.append(span);
+    });
     return graphContainer;
 }
diff --git a/Engine/Extras/RoboMerge/v3/public/js/main.js b/Engine/Extras/RoboMerge/v3/public/js/main.js
--- a/Engine/Extras/RoboMerge/v3/public/js/main.js
+++ b/Engine/Extras/RoboMerge/v3/public/js/main.js
@@ -121,7 +121,7 @@ function renderBranchList(data) {
 		let graphObj = branchGraphs.get(botName);
 		const branchSpecCl = botState.lastBranchspecCl;
 		if (!graphObj || graphObj.cl !== branchSpecCl) {
-			graphObj = {elem: showFlowGraph(branches, botName), cl: branchSpecCl};
+			graphObj = {elem: showFlowGraph(branches, {singleBotName: botName}), cl: branchSpecCl};
 			branchGraphs.set(botName, graphObj);
 		}
 	
diff --git a/Engine/Extras/RoboMerge/v3/public/js/preview.js b/Engine/Extras/RoboMerge/v3/public/js/preview.js
--- a/Engine/Extras/RoboMerge/v3/public/js/preview.js
+++ b/Engine/Extras/RoboMerge/v3/public/js/preview.js
@@ -4,19 +4,25 @@
 function doit(query) {
 	$.get(query)
 	.then(data => {
-		const botNames = new Set
 		const allBranches = JSON.parse(data).allBranches
-		for (const branch of allBranches) {
-			botNames.add(branch.bot)
+		if (allBranches.length === 0) {
+			throw new Error('no branches found in shelf')
 		}
-		$('#graph').append(showFlowGraph(allBranches, []));
+		const botNames = new Set(allBranches.map(b => b.bot))
+
+		$('#graph').append(showFlowGraph(allBranches, {
+			botNames: [], ...parseOptions(location.search)
+		}));
 		$('.bots').html([...botNames].map(s => `<tt>${s.toLowerCase()}</tt>`).join(', '))
 		$('#success-panel').show();
 	})
 	.catch(error => {
 
 		const $errorPanel = $('#error-panel');
-		$('pre', $errorPanel).text(error.responseText.replace(/\t/g, '    '));
+		const errorMsg = error.responseText
+			? error.responseText.replace(/\t/g, '    ')
+			: 'Internal error: ' + error.message;
+		$('pre', $errorPanel).text(errorMsg);
 		$errorPanel.show();
 	});
 }
diff --git a/Engine/Extras/RoboMerge/v3/public/js/skip.js b/Engine/Extras/RoboMerge/v3/public/js/skip.js
--- a/Engine/Extras/RoboMerge/v3/public/js/skip.js
+++ b/Engine/Extras/RoboMerge/v3/public/js/skip.js
@@ -179,7 +179,8 @@ function doSkip(fullData, edgeName, branchCl, skipText) {
     let queryData = {
         // If the forced CL is the same as the conflict CL, the conflict will be skipped
         cl: branchCl,
-        reason: skipText
+        reason: skipText,
+        unblock: 'true'
     }
 
     let branchData = fullData.branch
diff --git a/Engine/Extras/RoboMerge/v3/public/preview.html b/Engine/Extras/RoboMerge/v3/public/preview.html
--- a/Engine/Extras/RoboMerge/v3/public/preview.html
+++ b/Engine/Extras/RoboMerge/v3/public/preview.html
@@ -37,5 +37,19 @@ Graph for <span class="bots"></span> from <a href="https://p4-swarm.epicgames.ne
 			</div>
 			<div id="graph"></div>
 		</div>
+
+		<div style="transform: scale(0.7);">
+			<h4>Options</h4>
+			<table class="table">
+				<thead><tr><th>Setting</th><th>Effect</th></tr></thead>
+				<tbody>
+					<tr><td><tt>hideDisconnected</tt></td><td>Hide nodes with no edges</td></tr>
+					<tr><td><tt>noGroups</tt></td><td>Relevant if showing multiple bots: suppress boxes grouping each bot</td></tr>
+					<tr><td><tt>showOnlyForced</tt></td><td>Omit any "on demand" (dotted gray) edges</td></tr>
+				</tbody>
+			</table>
+
+			<p>e.g. <tt>http://robomerge/preview/12345?noGroups&showOnlyForced</tt></p>
+		</div>
 	</div>
 </body>
diff --git a/Engine/Extras/RoboMerge/v3/src/common/p4util.ts b/Engine/Extras/RoboMerge/v3/src/common/p4util.ts
--- a/Engine/Extras/RoboMerge/v3/src/common/p4util.ts
+++ b/Engine/Extras/RoboMerge/v3/src/common/p4util.ts
@@ -169,12 +169,12 @@ export async function convertIntegrateToEdit(p4: PerforceContext, roboWorkspace:
 	}
 }
 
-export async function cleanWorkspaces(p4: PerforceContext, workspaceNames: string[]) {
+export async function cleanWorkspaces(p4: PerforceContext, workspaces: [string, string][]) {
 	const changes = await p4.get_pending_changes() as Change[]
 
-	const nameSet = new Set(workspaceNames);
-	const toRevert = <Change[]>[];
-	const ignoredChangeCounts = new Map<string, number>();
+	const nameSet = new Set(workspaces.map(([name, _]) => name))
+	const toRevert: Change[] = []
+	const ignoredChangeCounts = new Map<string, number>()
 	for (const change of changes) {
 		if (nameSet.has(change.client)) {
 			toRevert.push(change);
@@ -215,5 +215,5 @@ export async function cleanWorkspaces(p4: PerforceContext, workspaceNames: strin
 	}
 
 	p4utilsLogger.info('Resetting all workspaces to revision 0')
-	await Promise.all(workspaceNames.map(ws => p4.sync(ws, `//${ws}/...#0`)))
+	await Promise.all(workspaces.map(([name, root]) => p4.sync(name, root + '#0')))
 }
diff --git a/Engine/Extras/RoboMerge/v3/src/common/perforce.ts b/Engine/Extras/RoboMerge/v3/src/common/perforce.ts
--- a/Engine/Extras/RoboMerge/v3/src/common/perforce.ts
+++ b/Engine/Extras/RoboMerge/v3/src/common/perforce.ts
@@ -502,13 +502,23 @@ export class PerforceContext {
 	}
 
 	async latestChange(path: string): Promise<Change> {
-		// wait no longer than 5 seconds, retry up to 3 times
-		const args = ['-vnet.maxwait=5', '-r3', 'changes', '-l', '-ssubmitted', '-m1', path]
+
+		// temporarily waiting 30 seconds - filing ticket   - was: wait no longer than 5 seconds, retry up to 3 times
+		const args = ['-vnet.maxwait=30', '-r3', 'changes', '-l', '-ssubmitted', '-m1', path]
+
+		const startTime = Date.now()
+
 		const result = await this.execAndParse(null, args, {quiet: true, trace: true}, changeResultExpectedShape)
 		if (!result || result.length !== 1) {
 			throw new Error("Expected exactly one change")
 		}
 
+		const durationSeconds = (Date.now() - startTime) / 1000;
+
+		if (durationSeconds > 5.0) {
+			this.logger.warn(`p4.latestChange took ${durationSeconds}s`)
+		}
+
 		return result[0] as unknown as Change
 	}
 
diff --git a/Engine/Extras/RoboMerge/v3/src/common/webserver.ts b/Engine/Extras/RoboMerge/v3/src/common/webserver.ts
--- a/Engine/Extras/RoboMerge/v3/src/common/webserver.ts
+++ b/Engine/Extras/RoboMerge/v3/src/common/webserver.ts
@@ -296,7 +296,7 @@ export class WebServer {
 			})
 
 			response.writeHead(500)
-			response.end(`Error handling not implemented yet: ${err.toString()}`)
+			response.end(`Error handling not implemented yet: ${err.toString()}\n${err.stack}`)
 			return
 		}
 
diff --git a/Engine/Extras/RoboMerge/v3/src/robo/autobranchupdater.ts b/Engine/Extras/RoboMerge/v3/src/robo/autobranchupdater.ts
--- a/Engine/Extras/RoboMerge/v3/src/robo/autobranchupdater.ts
+++ b/Engine/Extras/RoboMerge/v3/src/robo/autobranchupdater.ts
@@ -7,6 +7,7 @@ import { P4_FORCE, PerforceContext, Workspace } from '../common/perforce';
 import { Bot } from './bot-interfaces';
 import { BranchDefs } from './branchdefs';
 import { GraphInterface } from './graph-interface';
+import * as p4util from '../common/p4util';
 
 const DISABLE = false
 
@@ -15,6 +16,16 @@ type AutoBranchUpdaterConfig = {
 	workspace: Workspace	// { directory:	process.env('ROBO_BRANCHSPECS_DIRECTORY')
 }							// , name:		process.env('ROBO_BRANCHSPECS_WORKSPACE') }
 
+type MirrorPaths = {
+	name: string
+	directory: string
+	stream: string
+	depotFolder: string
+	depotpath: string
+	realFilepath: string
+	mirrorFilepath: string
+}
+
 export class AutoBranchUpdater implements Bot {
 
 	private p4: PerforceContext
@@ -33,6 +44,8 @@ export class AutoBranchUpdater implements Bot {
 	static config: AutoBranchUpdaterConfig
 	static initialCl: number
 
+	tickCount = 0
+
 	constructor(private graphBot: GraphInterface, parentLogger: ContextualLogger) {
 		this.p4 = AutoBranchUpdater.p4!
 
@@ -75,6 +88,19 @@ export class AutoBranchUpdater implements Bot {
 	async start() {
 		this.isRunning = true;
 		this.abuLogger.info(`Began monitoring ${this.graphBot.branchGraph.botname} branch specs at CL ${this.lastCl}`);
+
+// this got lost in the mirror refactor
+
+		// if (!DISABLE && this.mirror) {
+		// 	const workspace = await this.p4.find_workspace_by_name(this.mirror.workspace.name)
+		// 	if (workspace.length === 0) {
+		// 		await this.p4.newWorkspace(this.mirror.workspace.name, {
+		// 			Stream: this.mirror.stream,
+		// 			Root: AutoBranchUpdater.config!.workspace.directory
+		// 		})
+		// 	}
+		// }
+
 	}
 
 	async tick() {
@@ -142,7 +168,29 @@ export class AutoBranchUpdater implements Bot {
 			this.abuLogger.info(`Branch spec change detected: reloading ${botname} from CL#${this.lastCl}`)
 
 			await this.graphBot.reinitFromBranchGraphsObject(result.config, result.branchGraphDef!)
-			await this.updateMirror()
+
+			const mirrorWorkspace = AutoBranchUpdater.getMirrorWorkspace(this.graphBot)
+			if (!mirrorWorkspace) {
+				return
+			}
+
+			for (let retryCount = 0; ; ++retryCount) {
+				try {
+					await this.updateMirror(mirrorWorkspace)
+					return
+				}
+				catch (err) {
+
+					if (retryCount < 5) {
+						this.abuLogger.warn(`Mirror file reloading error (retries ${retryCount}: ${err}`)
+						p4util.cleanWorkspaces(this.p4, [[mirrorWorkspace.name, mirrorWorkspace.depotpath]])
+					}
+					else {
+						this.abuLogger.error('Mirror file reloading failed, retries exhausted')
+						throw err
+					}
+				}
+			}
 		})
 	}
 
@@ -158,12 +206,7 @@ export class AutoBranchUpdater implements Bot {
 		return this.abuLogger.context
 	}
 
-	private async updateMirror() {
-		const workspace = AutoBranchUpdater.getMirrorWorkspace(this.graphBot)
-		if (!workspace) {
-			return
-		}
-
+	private async updateMirror(workspace: MirrorPaths) {
 		this.abuLogger.info("Updating branchmap mirror")
 
 		const stream = this.graphBot.branchGraph.config.mirrorPath[0]
@@ -185,7 +228,7 @@ export class AutoBranchUpdater implements Bot {
 		await this.p4.submit(workspace, cl)
 	}
 
-	static getMirrorWorkspace(bot: GraphInterface) {
+	static getMirrorWorkspace(bot: GraphInterface): MirrorPaths | null {
 		const mirrorPathBits = bot.branchGraph.config.mirrorPath
 		if (mirrorPathBits.length === 0) {
 			return null
@@ -196,12 +239,13 @@ export class AutoBranchUpdater implements Bot {
 		}
 
 		const directory = AutoBranchUpdater.config!.workspace.directory
+		const depotFolder = mirrorPathBits.slice(0, 2).join('/')
 		const depotpath = mirrorPathBits.join('/')
 		const realFilepath = path.join(directory, bot.filename)
 		const mirrorFilepath = path.join(directory, ...mirrorPathBits.slice(1))
 		return {
 			name: `${AutoBranchUpdater.config!.workspace.name}-${bot.branchGraph.botname}-mirror`,
-			depotpath, realFilepath, mirrorFilepath, directory
+			stream: mirrorPathBits[0], depotFolder, depotpath, realFilepath, mirrorFilepath, directory
 		}
 	}
 }
diff --git a/Engine/Extras/RoboMerge/v3/src/robo/bot-interfaces.ts b/Engine/Extras/RoboMerge/v3/src/robo/bot-interfaces.ts
--- a/Engine/Extras/RoboMerge/v3/src/robo/bot-interfaces.ts
+++ b/Engine/Extras/RoboMerge/v3/src/robo/bot-interfaces.ts
@@ -1,9 +1,10 @@
 // Copyright Epic Games, Inc. All Rights Reserved.
 
 import { Branch, BranchArg, BranchGraphInterface, OperationResult, StompVerification } from "./branch-interfaces";
-import { BlockagePauseInfoMinimal, BranchStatus, ReconsiderArgs } from "./status-types"
+import { BlockagePauseInfoMinimal, BranchStatus } from "./status-types"
 import { TickJournal } from "./tick-journal";
 
+
 export interface Bot {
 	start(): Promise<void>;
 	tick(): Promise<boolean>;
@@ -16,6 +17,8 @@ export interface Bot {
 
 	lastCl: number
 
+	tickCount: number
+
 	tickJournal?: TickJournal;
 }
 
@@ -38,14 +41,20 @@ export interface IPCControls {
 	unblock(reason: string): void
 
 	pause(message: string, owner: string): void
-	unpause(): void
+	unpause(requester: string): void
 
 	reconsider(instigator: string, changeCl: number, additionalArgs?: Partial<ReconsiderArgs>): void
 
 	acknowledge(acknowledger: string, changeCl: number) : OperationResult
 	unacknowledge(changeCl: number) : OperationResult
 	
-	forceSetLastClWithContext(value: number, culprit: string, reason: string): number
+	forceSetLastClWithContext(value: number, culprit: string, reason: string, unblock?: boolean): number
+}
+
+export interface EdgeBotInterface {
+	get ipcControls(): IPCControls
+
+	bypassGateWindow(sense: boolean): void
 }
 
 export interface NodeBotInterface extends NodeBotIPC {
@@ -55,8 +64,35 @@ export interface NodeBotInterface extends NodeBotIPC {
 	readonly branchGraph: BranchGraphInterface
 	//readonly edges: Map<string, EdgeBot>
 	hasEdge(branchName: BranchArg): boolean
-	getEdgeIPCControls(branchName: BranchArg): IPCControls | undefined
+	getEdgeIPCControls(branchName: BranchArg): IPCControls | null
+	getImmediateEdge(branchName: BranchArg): EdgeBotInterface | null
 
 	getNumConflicts(): number
 	applyStatus(status: BranchStatus): void
+
+}
+
+// queued change vs reconsider:
+
+//  reconsider is the name of user facing 'queue change' operation, but also the top level call to add to the queue,
+//  used by stomp etc.
+
+// slightly sneaky that since NodeBot.reconsider's additional args are optional, it fulfils both of these, i.e.
+// implements the base bot reconsider function, and it also the 'queue change' function.
+
+// ReconsiderArgs are the additional args used by stomp etc. I'm adding commandOverride to them, even thoguh that's a
+// slightly different pattern (unless it breaks something)
+
+export type ReconsiderArgs = {
+	additionalFlags: string[]
+	workspace: string
+	targetBranchName: string
+	description: string
+	commandOverride: string
+}
+
+export type QueuedChange = Partial<ReconsiderArgs> & {
+	cl: number
+	who: string
+	timestamp: number 
 }
diff --git a/Engine/Extras/RoboMerge/v3/src/robo/branch-interfaces.ts b/Engine/Extras/RoboMerge/v3/src/robo/branch-interfaces.ts
--- a/Engine/Extras/RoboMerge/v3/src/robo/branch-interfaces.ts
+++ b/Engine/Extras/RoboMerge/v3/src/robo/branch-interfaces.ts
@@ -71,10 +71,10 @@ export interface BranchGraphInterface {
 	_computeReachableFrom(visited: Set<Branch>, flowKey: string, branch: Branch): Set<string>
 
 	getBranch(name: string): Branch | undefined
-	getBranchNames(): string
+	getBranchNames(): string[]
 }
 
-export interface Branch extends BranchBase {
+export interface EditableBranch extends BranchBase {
 	bot?: NodeBotInterface
 	parent: BranchGraphInterface
 	workspace: RoboWorkspace
@@ -97,6 +97,8 @@ export interface Branch extends BranchBase {
 	isMonitored: boolean // property
 }
 
+export type Branch = Readonly<EditableBranch>
+
 export interface Target {
 	branchName: string
 	mergeMode: string
diff --git a/Engine/Extras/RoboMerge/v3/src/robo/branchdefs.ts b/Engine/Extras/RoboMerge/v3/src/robo/branchdefs.ts
--- a/Engine/Extras/RoboMerge/v3/src/robo/branchdefs.ts
+++ b/Engine/Extras/RoboMerge/v3/src/robo/branchdefs.ts
@@ -20,7 +20,6 @@ export interface BotConfig {
 	reportToBuildHealth: boolean
 	mirrorPath: string[]
 	alias: string // alias if we need to mask name of bot in commands
-	emailDomainWhitelist: string[]
 	badgeUrlOverride: string
 	branchNamesToIgnore: string[]
 
@@ -38,7 +37,6 @@ export interface BranchBase {
 	flowsTo: string[]
 	forceFlowTo: string[]
 	defaultFlow: string[]
-	whitelist: string[]
 	resolver: string | null
 	aliases: string[]
 	badgeProject: string | null
@@ -113,6 +111,8 @@ type EdgeOptionFields = CommonOptionFields & {
 
 	terminal: boolean // changes go along terminal edges but no further
 	doHackyOkForGithubThing: boolean
+
+	implicitCommands: string[]
 }
 
 export type NodeOptions = Partial<NodeOptionFields>
@@ -223,7 +223,6 @@ export class BranchDefs {
 			reportToBuildHealth: false,
 			mirrorPath: [],
 			alias: '',
-			emailDomainWhitelist: [],
 			branchNamesToIgnore: [],
 			macros: {},
 			badgeUrlOverride: ''
diff --git a/Engine/Extras/RoboMerge/v3/src/robo/branchgraph.ts b/Engine/Extras/RoboMerge/v3/src/robo/branchgraph.ts
--- a/Engine/Extras/RoboMerge/v3/src/robo/branchgraph.ts
+++ b/Engine/Extras/RoboMerge/v3/src/robo/branchgraph.ts
@@ -2,15 +2,15 @@
 
 import { setDefault } from '../common/helper';
 import { BranchSpec } from '../common/perforce';
-import { Branch, BranchGraphInterface } from './branch-interfaces';
+import { Branch, BranchGraphInterface, EditableBranch } from './branch-interfaces';
 import { BotConfig, BranchDefs, BranchGraphDefinition, EdgeOptions, EdgeProperties, IntegrationMethod } from './branchdefs';
 import { NodeOptions, calculateStream } from './branchdefs';
 
 type ConfigBlendMode = 'override' | 'accumulate'
 
 export class BranchGraph implements BranchGraphInterface {
-
-	branches: Branch[]
+	// would ideally create branch list then assign it to a non-editable
+	branches: EditableBranch[]
 
 // should have separate type for processed edge info, but this will do for now
 	edges: EdgeProperties[]
@@ -48,7 +48,7 @@ if (botname === '__TEST__') {
 		return this.names.get(name.toUpperCase())
 	}
 	getBranchNames() {
-		return [...this.names.keys()].join(',')
+		return [...this.names.keys()]
 	}
 
 	private propagateBotPropertiesToNodes() {
@@ -208,7 +208,7 @@ if (botname === '__TEST__') {
 
 		// make a branch data entry
 		let nameUpper = name.toUpperCase()
-		let branch: Branch = {
+		let branch: EditableBranch = {
 			name: name,
 			parent: this,
 			workspace: options.workspace || null,
@@ -229,7 +229,6 @@ if (botname === '__TEST__') {
 			blockAssetTargets: new Set<string>(),
 			enabled: !options.disabled,
 			allowDeadend: !options.disallowDeadend,
-			whitelist: options.whitelist || [],
 			resolver: options.resolver || null,
 			convertIntegratesToEdits: (options.integrationMethod || this.config.defaultIntegrationMethod) === IntegrationMethod.CONVERT_TO_EDIT,
 			visibility: options.visibility || this.config.visibility,
diff --git a/Engine/Extras/RoboMerge/v3/src/robo/conflicts.ts b/Engine/Extras/RoboMerge/v3/src/robo/conflicts.ts
--- a/Engine/Extras/RoboMerge/v3/src/robo/conflicts.ts
+++ b/Engine/Extras/RoboMerge/v3/src/robo/conflicts.ts
@@ -235,7 +235,7 @@ export class Conflicts {
 		const resolvedConflicts: PersistentConflict[] = []
 
 		for (const conflict of this.conflicts) {
-			// If our bot CL is below the conflict CL, we can assume we haven't resolve this conflict
+			// If our bot CL is below the conflict CL, we can assume we haven't resolved this conflict
 			if (!conflict.targetBranchName || conflict.targetBranchName.toUpperCase() !== targetBranchStr || edgeBot.lastCl < conflict.cl) {
 				remainingConflicts.push(conflict)
 				continue
diff --git a/Engine/Extras/RoboMerge/v3/src/robo/edgebot.ts b/Engine/Extras/RoboMerge/v3/src/robo/edgebot.ts
--- a/Engine/Extras/RoboMerge/v3/src/robo/edgebot.ts
+++ b/Engine/Extras/RoboMerge/v3/src/robo/edgebot.ts
@@ -2,13 +2,12 @@
 
 import { ContextualLogger } from "../common/logger";
 import { Recipients } from "../common/mailer";
-import { Change, ChangelistStatus, ClientSpec, ConflictedResolveNFile, EditOwnerOpts, IntegrationSource }  from "../common/perforce";
-import { ResolveResult, IntegrationTarget, isExecP4Error, OpenedFileRecord, PerforceContext, EXCLUSIVE_CHECKOUT_REGEX } from "../common/perforce";
-import { convertIntegrateToEdit } from "../common/p4util";
+import { Change, ClientSpec, ConflictedResolveNFile, EditOwnerOpts, IntegrationSource }  from "../common/perforce";
+import { IntegrationTarget, isExecP4Error, OpenedFileRecord, PerforceContext, EXCLUSIVE_CHECKOUT_REGEX } from "../common/perforce";
 import { VersionReader } from "../common/version";
-import { IPCControls } from "./bot-interfaces";
+import { EdgeBotInterface, IPCControls, ReconsiderArgs } from "./bot-interfaces";
 import { PauseState } from "./state-interfaces";
-import { BlockagePauseInfo, ReconsiderArgs } from "./status-types";
+import { BlockagePauseInfo, BlockagePauseInfoMinimal } from "./status-types";
 import { AlreadyIntegrated, Branch, ChangeInfo, ConflictingFile, Failure, MergeAction, MergeMode, PendingChange } from "./branch-interfaces";
 import { PersistentConflict } from "./conflict-interfaces";
 import { BotEventTriggers } from "./events";
@@ -38,9 +37,57 @@ const JIRA_REGEX = /^\s*#jira\s+(.*)/i
 
 const MAX_CONFLICTS_TO_LIST = 5
 
-export type ResolveResultDetail = 'quick' | 'detailed'
+type ResolveResultDetail = 'quick' | 'detailed'
+
+type EdgeIntegrationResult = 'ok' | 'shelved' | 'nothing to do' | 'skipped' | 'error'
+
+class EdgeIntegrationDetails {
+	result: EdgeIntegrationResult
+	message = ''
+	cl = -1 // either submitted or shelved cl (may need Change in shelved case, let's see)
+
+	constructor(result: EdgeIntegrationResult, arg?: string | number) {
+		this.result = result
+		switch (typeof arg) {
+			case 'string': this.message = arg as string; break
+			case 'number': this.cl = arg as number; break
+		}
+	}
+}
+
+/** results from integrating along available outgoing edges */
+export class EdgeMergeResults {
+	addIntegrationResult(action: MergeAction, result: EdgeIntegrationDetails) {
+		this.edges.set(action.branch.upperName, result)
+		if (result.result === 'error') {
+			this.errors.push(result.message || 'no error message!')
+		}
+	}
+
+	allChanges() {
+		return [...this.edges].map(([_, v]) => v.cl).filter(n => n > 0)
+	}
+
+	markSkipped(target: Branch) {
+		this.edges.set(target.upperName, new EdgeIntegrationDetails('skipped'))
+	}
+
+	wasSkipped(target: Branch) {
+		return this.edges.has(target.upperName) && this.edges.get(target.upperName)!.result === 'skipped'
+	}
+
+	/** shelved or committed cl */
+	getTargetCl(target: Branch) {
+		const val = this.edges.get(target.upperName)
+		return val && val.cl
+	}
+
+	// tuples of target name, result, opt change
+	edges = new Map<string, EdgeIntegrationDetails>()
+	// accumulated errors for all edges
+	errors: string[] = []
+}
 
-export type EdgeIntegrationResult = string[] | Change
 
 class EdgeBotImpl extends PerforceStatefulBot {
 	readonly graphBotName: string
@@ -69,7 +116,10 @@ class EdgeBotImpl extends PerforceStatefulBot {
 		settings: Context,
 		matchingNodeConflict?: PersistentConflict) 
 	{
-		super(settings, options.initialCL || defaultLastCl)
+		super(settings, reason => {
+			this.edgeBotLogger.info('Queueing unblock, because ' + reason)
+			this.queueUnblock()
+		}, options.initialCL || defaultLastCl)
 
 		this.sourceBranch = sourceNode.branch
 		this.branch = targetBranch
@@ -120,6 +170,10 @@ class EdgeBotImpl extends PerforceStatefulBot {
 	get doHackyOkForGithubThing() {
 		return !!this.options.doHackyOkForGithubThing
 	}
+	get implicitCommands() {
+		return this.options.implicitCommands || []
+	}
+
 	get isTerminal() {
 		return !!this.options.terminal
 	}
@@ -168,22 +222,27 @@ class EdgeBotImpl extends PerforceStatefulBot {
 
 	}
 
-	updateLastCl(changesFetched: Change[], changeIndex: number, targetCl?: number) {
+	private updateLastCl(changesFetched: Change[], changeIndex: number, targetCl?: number) {
 		this.gate.updateLastCl(changesFetched, changeIndex, targetCl)
+		this.logger.verbose(`updating last cl to ${this.gate.lastCl} (${changesFetched.length}, ${changeIndex}, ${targetCl})`)
 		super._forceSetLastCl_NoReset(this.gate.lastCl)
 	}
 
-	private async getPerforceRequestResultFromCL(changelist: number, path?: string, changelistStatus?: ChangelistStatus) : Promise<EdgeIntegrationResult> {
-		const result = await this._getChange(changelist, path, changelistStatus)
-
-		if (!result.changes || !result.changes[0]) {
-			// swallowing actual error! should probably fix
-			return [`Failed to get ${changelistStatus ? changelistStatus : 'submitted'} CL ${changelist} (errors = ${result.errors}`]
+	onNodeProcessedChange(changesFetched: Change[], changeIndex: number, results: EdgeMergeResults) {
+		const change = changesFetched[changeIndex]
+		if (this.isAvailable && this.lastCl < change.change && !results.wasSkipped(this.targetBranch)) {
+			const targetCl = results.getTargetCl(this.targetBranch)
+			this.updateLastCl(changesFetched, changeIndex, targetCl)
 		}
+		else {
+			this.gate.calcNumChangesRemaining(changesFetched, changeIndex)
+		}
+	}
 
-		return result.changes[0]
+	block(blockageInfo: BlockagePauseInfoMinimal, pauseDurationSeconds?: number) {
+		super.block(blockageInfo, pauseDurationSeconds)
 	}
-	
+
 	private analyzeConflict(unresolved: ConflictedResolveNFile[]) {
 		const results: ConflictingFile[] = []
 
@@ -275,7 +334,7 @@ class EdgeBotImpl extends PerforceStatefulBot {
 		this.block(this.createEdgeBlockageInfo(failure, pending), pauseDurationSeconds)
 	}
 
-	performMerge(info: ChangeInfo, target: MergeAction, convertIntegratesToEdits: boolean): Promise<EdgeIntegrationResult> {
+	performMerge(info: ChangeInfo, target: MergeAction, convertIntegratesToEdits: boolean, additionalDescription: string): Promise<EdgeIntegrationDetails> {
 		// make sure we come back in here afterwords
 		this._log_action(`Merging CL ${info.cl} via ${this.fullName} (${target.mergeMode})`)
 
@@ -308,18 +367,21 @@ class EdgeBotImpl extends PerforceStatefulBot {
 		// if the owner is specifically overridden (can be by reconsider, resolver or manual tag)
 		const overriddenOwner = getIntegrationOwner(this.branch, info.owner)
 
-		let authorTag = info.authorTag
 		if (overriddenOwner) {
 			description += `#ROBOMERGE-OWNER: ${overriddenOwner}\n`
-			// put an author tag if the owner is overridden, in case the owner ends up committing in their own name
-			authorTag = authorTag || info.author
 		}
 
-		if (authorTag) {
+		// keep track of author in a tag in case transfering onwership of the changelist fails
+		const authorTag = info.authorTag || info.author
+		if (authorTag !== 'robomerge') {
 			description += `#ROBOMERGE-AUTHOR: ${authorTag}\n`
 		}
 
-		if (this.doHackyOkForGithubThing && !info.hasOkForGithubTag) {
+		if (this.doHackyOkForGithubThing/* && !info.hasOkForGithubTag */) {
+			const descriptionLines = description.split('\n')
+			description = descriptionLines
+				.filter(s => !s.toLowerCase().startsWith('#okforgithub'))
+				.join('\n')
 			description += '#okforgithub ignore\n'
 		}
 
@@ -397,7 +459,7 @@ class EdgeBotImpl extends PerforceStatefulBot {
 			description += '#ROBOMERGE-EDIGRATE\n'
 		}
 
-		target.description = description
+		target.description = description + additionalDescription
 
 		// do the integration
 		return this.integrate(info, target)
@@ -407,7 +469,16 @@ class EdgeBotImpl extends PerforceStatefulBot {
 		this.currentIntegrationStartTimestamp = -1
 	}
 
-	private async integrate(info: ChangeInfo, target: MergeAction) : Promise<EdgeIntegrationResult> {
+	public bypassGateWindow(sense: boolean) {
+		this.gate.bypass = sense
+		this.gate.persist()
+	}
+
+	queueUnblock() {
+		this.sourceNode.queueEdgeUnblock(this.targetBranch.upperName)
+	}
+
+	private async integrate(info: ChangeInfo, target: MergeAction) : Promise<EdgeIntegrationDetails> {
 		let to_integrate = info.cl
 		this._log_action(`Integrating CL ${to_integrate} to ${this.targetBranch.name}`)
 
@@ -457,7 +528,7 @@ class EdgeBotImpl extends PerforceStatefulBot {
 				// integration not necessary
 				this.edgeBotLogger.info(msg)
 				await this.p4.deleteCl(this.targetBranch.workspace, changenum)
-				return [ msg ]
+				return new EdgeIntegrationDetails('nothing to do', msg)
 			}
 		}
 
@@ -481,8 +552,10 @@ class EdgeBotImpl extends PerforceStatefulBot {
 		
 		await this.handleIntegrationError(failure, pending)
 		// Send to source node to facilitate notification handling
-		await this.sourceNode.handleMergeFailure(failure, pending)
-		return errors
+		if (this.sourceNode.handleMergeFailure(failure, pending)) {
+			this.gate.onBlockage()
+		}
+		return new EdgeIntegrationDetails('error', errors.join('\n'))
 	}
 
 	private async revertAndDelete(cl: number) {
@@ -508,66 +581,13 @@ class EdgeBotImpl extends PerforceStatefulBot {
 		return this.p4.resolve(this.targetBranch.workspace, changelist, mergeMode, detail === 'quick')
 	}
 
-// code restored that was deleted in CL#8057353 (from nodebot.ts) - still needs work
-	private async roboShelve(result: ResolveResult, toShelve: PendingChange) : Promise<EdgeIntegrationResult> {
-		const targetBranch = toShelve.action.branch
-		const targetEdge = this
-		let errorStr: string | null = null
-		let errorDetails: string | null = null
-
-		if (result.hasConflict()) {
-			// @todo Improve resolve result
-			errorStr = 'Unable to RoboShelve due to conflicts',
-			errorDetails = result.getConflictsText()
-		}
-		else {
-			this.edgeBotLogger.info(`Converting integrations in CL ${toShelve.change.cl} into edits`)
-			try {
-				await convertIntegrateToEdit(this.p4, targetBranch.workspace!, toShelve.newCl)
-			}
-			catch (err) {
-				errorStr = 'Conversion from integrates to edits failed'
-				errorDetails = err.toString()
-			}
-		}
-
-		if (errorStr) {
-			// Failed!
-
-			if (toShelve.change.userRequest) {
-				// send email in parallel
-				// this._sendGenericEmail(new Recipients(toShelve.change.owner!), 'RoboShelf error', errorStr, errorDetails!);
-				await targetEdge.revertAndDelete(toShelve.newCl)
-			}
-			else {
-
-				const failure: Failure = {kind: 'Conversion to edits failure', description: errorStr}
-				// should probably have custom handling, but try this for tests 
-				this.handleIntegrationError(failure, toShelve)
-
-				await this.sourceNode.handleMergeFailure(failure, toShelve)
-			}
-			
-			return [`${errorStr}: ${errorDetails}`]
-		}
-		else {
-
-			await targetEdge.submitChangelist(toShelve)
-			return this.getPerforceRequestResultFromCL(toShelve.newCl)
-
-// should shelve if manual or non auto target?
-			// await targetEdge.shelveChangelist(toShelve, 'RoboShelf has prepared this shelf for your consideration')
-			// return this.getPerforceRequestResultFromCL(toShelve.newCl, this.getChangePaths(targetBranch), "shelved")
-
-		}
-	}
 
 	/**
 	 * Perform a resolve of the change described by pending
 	 * @param {PendingChange} pending Pending change to resolve and submit
 	 * @param {number} [submitRetries = 3] number of times 'targetEdge.submitChangelist' can fail gracefully and come back to this method
 	 */
-	async _resolveChangelist(pending: PendingChange, submitRetries  = 3) : Promise<EdgeIntegrationResult> {
+	async _resolveChangelist(pending: PendingChange, submitRetries  = 3) : Promise<EdgeIntegrationDetails> {
 		// do a resolve with P4
 		this._log_action(`Resolving CL ${pending.change.cl} against ${this.targetBranch.name}`)
 		let detail: ResolveResultDetail = 'detailed'
@@ -580,7 +600,7 @@ class EdgeBotImpl extends PerforceStatefulBot {
 				this.sourceNode.emailShelfRequester(pending)
 			}
 			
-			return this.getPerforceRequestResultFromCL(pending.newCl, pending.action.branch.rootPath, "shelved")
+			return new EdgeIntegrationDetails('shelved', pending.newCl)
 		}
 
 		let failure: Failure | null = null
@@ -603,27 +623,18 @@ class EdgeBotImpl extends PerforceStatefulBot {
 		}
 
 		if (!failure) {
-			if (this.sourceBranch.convertIntegratesToEdits) {
-				return await this.roboShelve(result, pending)
+			const result = await this.submitChangelist(pending, submitRetries)
+			if (result.result !== 'error') {
+				return result
 			}
 
-
-			// No conflicts - let's do this!
-			const error = await this.submitChangelist(pending, submitRetries)
-
-			if (typeof(error) === 'string') {
-				failure = { kind: 'Commit failure', description: error }
-			}
+			failure = { kind: 'Commit failure', description: result.message }
 		}
 
-		if (failure) {
-			await this.handlePostIntegrationFailure(failure, pending)
-			await this.sourceNode.handleMergeFailure(failure, pending, true)
-
-			return [`${failure.kind}: ${failure.description}`]
-		}
+		await this.handlePostIntegrationFailure(failure, pending)
+		await this.sourceNode.handleMergeFailure(failure, pending, true)
 
-		return this.getPerforceRequestResultFromCL(pending.newCl)
+		return new EdgeIntegrationDetails('error', `${failure.kind}: ${failure.description}`)
 	}
 
 	/**
@@ -632,7 +643,7 @@ class EdgeBotImpl extends PerforceStatefulBot {
 	 * @param {Number} [resolveRetries=3] Number of times to catch non-fatal errors and attempt the resolve -> submit chain again.
 	 * @todo In a perfect world, parameter 'resolveRetries' wouldn't exist and an orchastration method would handle this in a loop.
 	 */
-	async submitChangelist(pending: PendingChange, resolveRetries = 3): Promise<void|string> {
+	async submitChangelist(pending: PendingChange, resolveRetries = 3): Promise<EdgeIntegrationDetails> {
 		const info = pending.change
 		const target = pending.action
 		const changenum = pending.newCl
@@ -640,22 +651,21 @@ class EdgeBotImpl extends PerforceStatefulBot {
 		// try to submit
 		this._log_action(`Submitting CL ${changenum} by ${info.author}`)
 		const result = await this.p4.submit(this.targetBranch.workspace, changenum)
-		if (typeof(result) === "string") {
+		if (typeof result === "string") {
 			// an error occurred while submitting
-			return result;
+			return new EdgeIntegrationDetails('error', result);
 		}
 		const finalCl = result;
 		if (finalCl === 0) {
 			if (resolveRetries >= 1) {
 				// we need to resolve again
 				this.edgeBotLogger.info(`Detected concurrent changes. Re-resolving CL ${changenum}`)
-				await this._resolveChangelist(pending, --resolveRetries)
-				return
+				return await this._resolveChangelist(pending, --resolveRetries)
 			}
 			else {
 				const msg = `Hit maximum number of retries when trying to submit changelists for pending change ${pending.change.cl}`
 				this.edgeBotLogger.warn(msg)
-				return msg
+				return new EdgeIntegrationDetails('error', msg)
 			}
 		}
 
@@ -692,7 +702,7 @@ class EdgeBotImpl extends PerforceStatefulBot {
 						if (this.targetBranch.config.additionalSlackChannelForBlockages) {
 							postMessageToChannel(msg, this.targetBranch.config.additionalSlackChannelForBlockages, SlackMessageStyles.WARNING)
 						}
-						return
+						return new EdgeIntegrationDetails('ok', finalCl)
 					}
 					else if (output.match(/You don't have permission for this operation/)) {
 						errPreface += ' (This error is expected in non-prod instances.)'
@@ -702,6 +712,7 @@ class EdgeBotImpl extends PerforceStatefulBot {
 				this.edgeBotLogger.printException(exception, errPreface);
 			}
 		}
+		return new EdgeIntegrationDetails('ok', finalCl)
 	}
 
 	private async shelveChangelist(pending: PendingChange, reason?: string) {
@@ -817,7 +828,8 @@ class EdgeBotImpl extends PerforceStatefulBot {
 		if (this.pauseState.isBlocked()) {
 			// see if this change matches our stop CL
 			if (info.source_cl === this.pauseState.blockagePauseInfo!.sourceCl) {
-				this.unblock(`finding CL ${info.source_cl} merged to ${info.branch.name} in CL ${info.cl}`)
+				this.edgeBotLogger.info(`Queueing unblock due to finding CL ${info.source_cl} merged to ${info.branch.name} in CL ${info.cl}`)
+				this.queueUnblock()
 			}
 		}
 	}
@@ -855,6 +867,8 @@ class EdgeBotImpl extends PerforceStatefulBot {
 		status.rootPath = this.targetBranch.rootPath
 
 		status.last_cl = this.lastCl
+		
+		// if (this.isForceFlow)
 		status.num_changes_remaining = this.gate.numChangesRemaining
 
 		status.is_active = this.isActive
@@ -888,8 +902,12 @@ class EdgeBotImpl extends PerforceStatefulBot {
 		return status
 	}
 
-	public forceSetLastClWithContext(value: number, culprit: string, reason: string) {
-		const prevValue = this.forceSetLastCl(value)
+	public forceSetLastClWithContext(value: number, culprit: string, reason: string, unblock?: boolean) {
+		const prevValue = this._forceSetLastCl_NoReset(value)
+
+		if (unblock) {
+			this.unblock(`Also unblocking when skipping to cl ${value}`)
+		}
 
 		// trigger events
 		this.sourceNode.onForcedLastCl(this.displayName, value, prevValue, culprit, reason)
@@ -929,7 +947,10 @@ abstract class EdgeBotEntryPoints implements IPCControls {
 
 /* This wrapper class allows us to control what is accessing the EdgeBotImpl internals, and
  *	allows us to monitor when an edge is active */
-export class EdgeBot extends EdgeBotEntryPoints {
+export class EdgeBot
+	extends EdgeBotEntryPoints
+	implements EdgeBotInterface
+{
 	private impl: EdgeBotImpl
 	displayName: string
 	fullName: string
@@ -1002,8 +1023,16 @@ export class EdgeBot extends EdgeBotEntryPoints {
 		this.impl.preIntegrate(cl)
 	}
 
-	updateLastCl(changesFetched: Change[], changeIndex: number, targetCl?: number) {
-		this.impl.updateLastCl(changesFetched, changeIndex, targetCl)
+	onNodeProcessedChange(changesFetched: Change[], changeIndex: number, results: EdgeMergeResults) {
+		this.impl.onNodeProcessedChange(changesFetched, changeIndex, results)
+	}
+
+	bypassGateWindow(sense: boolean) {
+		this.impl.bypassGateWindow(sense)
+	}
+
+	queueUnblock() {
+		this.impl.queueUnblock()
 	}
 
 	/* Mirrored Variables */
@@ -1027,6 +1056,9 @@ export class EdgeBot extends EdgeBotEntryPoints {
 	get targetBranch() {
 		return this.impl.targetBranch
 	}
+	get implicitCommands() {
+		return this.impl.implicitCommands
+	}
 
 	get ipcControls(): IPCControls {
 		return {
diff --git a/Engine/Extras/RoboMerge/v3/src/robo/gate.ts b/Engine/Extras/RoboMerge/v3/src/robo/gate.ts
--- a/Engine/Extras/RoboMerge/v3/src/robo/gate.ts
+++ b/Engine/Extras/RoboMerge/v3/src/robo/gate.ts
@@ -9,6 +9,7 @@ import { Context } from "./settings"
 
 
 const GATE_INFO_KEY = 'gate-info'
+const BYPASS_KEY = 'bypass'
 
 
 type GateOptions = Partial<CommonOptionFields>
@@ -96,12 +97,21 @@ type EventTriggersAndStuff = GateEventContext & {
 	eventTriggers: BotEventTriggers
 }
 
+type CurrentWindowStats = {
+	targetCl: number
+	pending: number
+	integrated: number
+	blockages: number
+}
+
 export class Gate {
 	constructor(	private eventContext: EventTriggersAndStuff | DummyEventTriggers,
 					private context: GateContext, private persistence?: Context) {
 		this.loadFromPersistence()
 	}
 
+	bypass = false
+
 	private processGateChange(newGate: GateInfo) {
 		// degrees of freedom:
 		// 	order of lastCl/gate.cl/newGate.cl (L/G/N)
@@ -118,6 +128,8 @@ export class Gate {
 			if (newGate.cl <= this.currentGateInfo.cl) {
 				// report if we were catching up and the replacement gate means we're now caught up
 				if (this.lastCl < this.currentGateInfo.cl && this.lastCl >= newGate.cl) {
+					this.context.logger.info('Caught up due to earlier gate ' +
+						`(previous: ${this.currentGateInfo.cl}, new: ${newGate.cl},lastCl = ${this.lastCl}`)
 					this.reportCaughtUp()
 				}
 				this.currentGateInfo = newGate
@@ -158,50 +170,60 @@ export class Gate {
 			return
 		}
 
-		const mostRecentGate = this.getMostRecentGate()
+		let dirty = false
+		try {
+			if (this.bypass && this.nowIsWithinAllowedCatchupWindow()) {
+				this.context.logger.info('Clearing bypass on entering an integration window')
+				this.bypass = false
+				dirty = true
+			}
 
-		const gateInfo = await getRequestedGateCl(this.context, mostRecentGate)
-		if (!gateInfo) {
-			if (this.currentGateInfo) {
-				// we were waiting for a gate, so unpause and clear gate info
-				if (this.isGateOpen()) {
-					this.reportCaughtUp()
+			const mostRecentGate = this.getMostRecentGate()
+
+			const gateInfo = await getRequestedGateCl(this.context, mostRecentGate)
+			if (!gateInfo) {
+				if (this.currentGateInfo) {
+					// we were waiting for a gate, so unpause and clear gate info
+					if (this.isGateOpen()) {
+						this.context.logger.info(`Gate ${this.currentGateInfo.cl} removed while catching up`)
+						this.reportCaughtUp()
+					}
+
+					this.currentGateInfo = null
+					this.queuedGates.length = 0
+					dirty = true
 				}
 
-				this.currentGateInfo = null
-				this.queuedGates.length = 0
-				this.persist()
+				return
 			}
 
-			return
-		}
-
-		let dirty = false
-		if (!mostRecentGate) {
-			// ooh look, our first gate
-			this.queuedGates.push(gateInfo)
-			dirty = true
-		}
-		else if (!gatesSame(mostRecentGate, gateInfo)) {
-			this.processGateChange(gateInfo)
-			dirty = true
-		}
+			if (!mostRecentGate) {
+				// ooh look, our first gate
+				this.queuedGates.push(gateInfo)
+				dirty = true
+			}
+			else if (!gatesSame(mostRecentGate, gateInfo)) {
+				this.processGateChange(gateInfo)
+				dirty = true
+			}
 
-		// sanity check
-		if (!this.currentGateInfo && this.queuedGates.length === 0) {
-			throw new Error('seem to have ignored gate!')
-		}
+			// sanity check
+			if (!this.currentGateInfo && this.queuedGates.length === 0) {
+				throw new Error('seem to have ignored gate!')
+			}
 
-		// if no current and in window, kick off new
-		if (!this.currentGateInfo && this.queuedGates.length > 0) {
-			if (this.tryNextGate(this.lastCl)) {
-				this.reportCatchingUp()
-				dirty = true
+			// if no current and in window, kick off new
+			if (!this.currentGateInfo && this.queuedGates.length > 0) {
+				if (this.tryNextGate(this.lastCl)) {
+					this.reportCatchingUp()
+					dirty = true
+				}
 			}
 		}
-
-		if (dirty) {
-			this.persist()
+		finally {
+			if (dirty) {
+				this.persist()
+			}
 		}
 	}
 
@@ -210,23 +232,32 @@ export class Gate {
 			this.context.logger.verbose(`${this.lastCl} -> ${cl} ${this.currentGateInfo.cl}`)
 		}
 
-		let adjustedCl: number | null = null
-		if (this.currentGateInfo && cl > this.currentGateInfo.cl && this.currentGateInfo.cl > this.lastCl) {
-			// we're waiting for currentGateInfo.cl but got a higher cl
-			// basically we drag lastCl forward to currentGateInfo.cl so isGateOpen returns false
-			// but if there are queued gates, we should have a go at moving to the next gate instead
+		if (!this.currentGateInfo || cl <= this.currentGateInfo.cl || this.currentGateInfo.cl <= this.lastCl) {
+			return null
+		}
 
-			// note tryNextGate can change currentGateInfo or set it to null
-			if (!this.tryNextGate(cl)) {
-				this.reportCaughtUp()
-				if (this.currentGateInfo) {
-					adjustedCl = this.currentGateInfo.cl
-					this.context.logger.info(`Adjusting cl to ${adjustedCl}`)
-					this.setLastCl(adjustedCl)
-				}
-			}
+		// we're waiting for currentGateInfo.cl but got a higher cl
+
+		// if there are no more gates, basically we drag lastCl forward to
+		// currentGateInfo.cl so isGateOpen returns false
+
+		this.context.logger.verbose(`Current gate info: ${this.currentGateInfo && this.currentGateInfo.cl}`)
+
+		if (this.tryNextGate(cl)) {
+			return null
 		}
-		return adjustedCl
+
+		// new gate, check again
+		if (!this.currentGateInfo || cl <= this.currentGateInfo.cl || this.currentGateInfo.cl <= this.lastCl) {
+			return null
+		}
+
+		this.reportCaughtUp()
+		this.context.logger.info(`Adjusting cl from ${this.lastCl} to ${this.currentGateInfo.cl} ` +
+														`to match gate due to encountering CL#${cl}`)
+		this.setLastCl(this.currentGateInfo.cl)
+		this.persist()
+		return this.currentGateInfo.cl
 	}
 
 	/**
@@ -249,7 +280,7 @@ export class Gate {
 			return false
 		}
 
-		if (!this.nowIsWithinAllowedCatchupWindow()) {
+		if (!this.bypass && !this.nowIsWithinAllowedCatchupWindow()) {
 			// want to make this an info log, but would spam every tick at the moment
 			// this.context.logger.verbose('delaying gate catch-up due to configured window')
 
@@ -281,16 +312,31 @@ export class Gate {
 		}
 
 		this.setLastCl(cl)
+
+		this.numChangesRemaining = this.calcNumChangesRemaining(changesFetched, changeIndex)
+
+		if (this.currentWindowStats) {
+			if (targetCl) {
+				this.currentWindowStats.targetCl = targetCl
+				++this.currentWindowStats.integrated
+			}
+			this.currentWindowStats.pending = this.numChangesRemaining
+		}
+
 		if (notifyCaughtUp) {
-			this.reportCaughtUp(targetCl)
+			this.reportCaughtUp()
 		}
+	}
 
-		this.numChangesRemaining = this.calcNumChangesRemaining(changesFetched, changeIndex)
+	onBlockage() {
+		if (this.currentWindowStats) {
+			++this.currentWindowStats.blockages
+		}
 	}
 
 	numChangesRemaining = 0
 
-	private calcNumChangesRemaining(changesFetched: Change[], changeIndex: number) {
+	calcNumChangesRemaining(changesFetched: Change[], changeIndex: number) {
 		const mostRecentGate = this.getMostRecentGate()
 		if (!mostRecentGate) {
 			return changesFetched.length - changeIndex - 1
@@ -348,9 +394,14 @@ export class Gate {
 		const closedMessage = this.getGateClosedMessage()
 		if (closedMessage) {
 			outStatus.gateClosedMessage = closedMessage
-			if (this.nextWindowOpenTime) {
-				outStatus.nextWindowOpenTime = this.nextWindowOpenTime
-			}
+		}
+
+		if (this.nextWindowOpenTime) {
+			outStatus.nextWindowOpenTime = this.nextWindowOpenTime
+		}
+
+		if (this.bypass) {
+			outStatus.windowBypass = true
 		}
 	}
 
@@ -435,13 +486,15 @@ export class Gate {
 					}
 				}
 				else {
-					consider(now.getUTCDay(), pane)
+					const nowDay = now.getUTCDay()
+					// covers window being before or after now
+					consider(nowDay, pane)
+					consider((nowDay + 1) % 7, pane)
 				}
 			}
 
 			// for now, only provide message if we have one non-inverted window
 			// more general solution to follow
-			// also not taking into account days of the week yet
 			this.nextWindowOpenTime = now
 			this.nextWindowOpenTime.setUTCHours(earliestHour - now.getUTCDay() * 24, 0, 0, 0);
 		}
@@ -461,16 +514,24 @@ export class Gate {
 			info: mostRecent,
 			changesRemaining: this.numChangesRemaining
 		})
+
+		this.currentWindowStats = Gate.newCurrentWindowStats()
 	}
 
-	private reportCaughtUp(targetCl?: number) {
+	private reportCaughtUp() {
+// log stats here for now, but if it works out, use the type in the event
+		this.context.logger.info('Catch-up window stats: ' + (this.currentWindowStats
+			? JSON.stringify(this.currentWindowStats) : '<null>'))
+
 		this.eventContext.eventTriggers.reportEndIntegratingToGate({
 			context: this.eventContext,
-			targetCl: targetCl || -1
+			targetCl: this.currentWindowStats ? this.currentWindowStats.targetCl : -1
 		})
+
+		this.currentWindowStats = null
 	}
 
-	private persist() {
+	persist() {
 		if (!this.persistence) {
 			return
 		}
@@ -484,6 +545,7 @@ export class Gate {
 		}
 		data.lastCl = this.lastCl
 		this.persistence.set(GATE_INFO_KEY, data)
+		this.persistence.set(BYPASS_KEY, this.bypass)
 	}
 
 	private loadFromPersistence() {
@@ -521,6 +583,7 @@ export class Gate {
 				this.queuedGates = saved.queued
 			}
 		}
+		this.bypass = !!this.persistence.get(BYPASS_KEY)
 	}
 
 	getEventContextForTests() {
@@ -537,6 +600,15 @@ export class Gate {
 
 	private currentGateInfo: GateInfo | null = null
 	private queuedGates: GateInfo[] = []
+
+	// info only
+
+	private static newCurrentWindowStats(): CurrentWindowStats
+	{
+		return { targetCl: -1, pending: -1, integrated: 0, blockages: 0 }
+	}
+
+	private currentWindowStats: CurrentWindowStats | null = null
 }
 
 /**
@@ -738,7 +810,9 @@ export async function runTests(parentLogger: ContextualLogger) {
 		if (middleIntegration) {
 			gate.preIntegrate(2); setLastCl(gate, 2); await gate.tick()
 		}
-		gate.preIntegrate(3); setLastCl(gate, 3); await gate.tick()
+		gate.preIntegrate(3);
+		assert(gate.lastCl < 3, 'exact gate cl gets integrated')
+		setLastCl(gate, 3); await gate.tick()
 
 		assert(et.beginCalls === 1 && et.endCalls === 1, 'catch-up notified')
 
diff --git a/Engine/Extras/RoboMerge/v3/src/robo/graph-interface.ts b/Engine/Extras/RoboMerge/v3/src/robo/graph-interface.ts
--- a/Engine/Extras/RoboMerge/v3/src/robo/graph-interface.ts
+++ b/Engine/Extras/RoboMerge/v3/src/robo/graph-interface.ts
@@ -9,4 +9,5 @@ export interface GraphInterface {
 
 	ensureStopping(): boolean
 	reinitFromBranchGraphsObject(config: BotConfig, branchGraphs: BranchGraphDefinition): Promise<void>
+	handleRequestedIntegrationsForAllNodes(): Promise<void>
 }
diff --git a/Engine/Extras/RoboMerge/v3/src/robo/graphbot.ts b/Engine/Extras/RoboMerge/v3/src/robo/graphbot.ts
--- a/Engine/Extras/RoboMerge/v3/src/robo/graphbot.ts
+++ b/Engine/Extras/RoboMerge/v3/src/robo/graphbot.ts
@@ -129,7 +129,9 @@ export class GraphBot implements GraphInterface, BotEventHandler {
 		for (const branch of this.branchGraph.branches) {
 			if (branch.enabled) {
 				const persistence = this.settings.getContext(branch.upperName)
-				branch.bot = new NodeBot(branch, this.mailer, this.externalUrl, this.eventTriggers, persistence, ubergraph)
+				branch.bot = new NodeBot(branch, this.mailer, this.externalUrl, this.eventTriggers, persistence, ubergraph,
+					() => this.handleRequestedIntegrationsForAllNodes()
+				)
 
 				if (branch.bot.getNumConflicts() > 0) {
 					hasConflicts = true
@@ -172,14 +174,17 @@ export class GraphBot implements GraphInterface, BotEventHandler {
 		postToRobomergeAlerts(msg)
 
 		if (this.branchGraph.branches.length !== 0) {
-			const workspaceNames = this.branchGraph.branches.map(branch => (branch.workspace as Workspace).name || (branch.workspace as string))
+			const workspaces = this.branchGraph.branches.map(branch => 
+				[(branch.workspace as Workspace).name || (branch.workspace as string),
+				branch.rootPath]) as [string, string][]
 			const mirrorWorkspace = AutoBranchUpdater.getMirrorWorkspace(this)
 			if (mirrorWorkspace) {
-				workspaceNames.push(mirrorWorkspace.name)
+				// add /... to match branches' rootPath format
+				workspaces.push([mirrorWorkspace.name, mirrorWorkspace.stream + '/...'])
 			}
 
 			this.botLogger.info('Cleaning all workspaces')
-			await p4util.cleanWorkspaces(this.p4, workspaceNames)
+			await p4util.cleanWorkspaces(this.p4, workspaces)
 		}
 
 		await this.startBotsAsync()
@@ -193,6 +198,20 @@ export class GraphBot implements GraphInterface, BotEventHandler {
 		this.crashRequested = msg
 	}
 
+	async handleRequestedIntegrationsForAllNodes() {
+		for (const branchName of this.branchGraph.getBranchNames()) {
+			const node = this.findNode(branchName)!
+			for (;;) {
+				const request = node.popRequestedIntegration()
+				if (!request) {
+					break
+				}
+				await node.processQueuedChange(request)
+				node.persistQueuedChanges()
+			}
+		}
+	}
+
 	private async startBotsAsync() {
 		if (!this.waitTime) {
 			throw new Error('runbots must be called before startBots')
@@ -253,10 +272,13 @@ export class GraphBot implements GraphInterface, BotEventHandler {
 				}
 				bot.isActive = false
 
-				if (ticked && bot.tickJournal) {
-					const nodeBot = bot as NodeBot
-					bot.tickJournal.monitored = nodeBot.branch.isMonitored
-					activity.set(nodeBot.branch.upperName, bot.tickJournal)
+				if (ticked) {
+					++bot.tickCount
+					if (bot.tickJournal) {
+						const nodeBot = bot as NodeBot
+						bot.tickJournal.monitored = nodeBot.branch.isMonitored
+						activity.set(nodeBot.branch.upperName, bot.tickJournal)
+					}
 				}
 
 				if (this._shutdownCb) {
@@ -270,6 +292,8 @@ export class GraphBot implements GraphInterface, BotEventHandler {
 				await new Promise(done => setTimeout(done, this.waitTime!))
 			}
 
+			await this.handleRequestedIntegrationsForAllNodes()
+
 			roboAnalytics!.reportActivity(this.branchGraph.botname, activity)
 			roboAnalytics!.reportMemoryUsage('main', process.memoryUsage().heapUsed)
 
diff --git a/Engine/Extras/RoboMerge/v3/src/robo/ipc.ts b/Engine/Extras/RoboMerge/v3/src/robo/ipc.ts
--- a/Engine/Extras/RoboMerge/v3/src/robo/ipc.ts
+++ b/Engine/Extras/RoboMerge/v3/src/robo/ipc.ts
@@ -3,7 +3,7 @@
 import { ContextualLogger, isNpmLogLevel, NpmLogLevelCompare, NpmLogLevelValues } from '../common/logger';
 import { getRunningPerforceCommands } from '../common/perforce';
 import { Trace } from '../new/graph';
-import { IPCControls, NodeBotInterface } from './bot-interfaces';
+import { IPCControls, EdgeBotInterface, NodeBotInterface } from './bot-interfaces';
 import { OperationResult } from './branch-interfaces';
 import { GraphBot } from './graphbot';
 import { RoboMerge } from './robo';
@@ -328,12 +328,14 @@ export class IPC {
 		// Some operations are the same for NodeBots and EdgeBots -- we can cheat a little here by
 		// creating a general variable
 		let generalOpTarget: IPCControls = bot
+		let edge: EdgeBotInterface | null = null
 		if (edgeName) {
-			const edge = bot.getEdgeIPCControls(edgeName.toUpperCase())
-			if (!edge) {
+			edge = bot.getImmediateEdge(edgeName)
+			const edgeIPC = edge && edge.ipcControls
+			if (!edgeIPC) {
 				return {statusCode: 400, message: `Node ${nodeName} does not have edge ${edgeName}`}
 			}
-			generalOpTarget = edge
+			generalOpTarget = edgeIPC
 		}
 		
 		let cl = NaN
@@ -346,7 +348,7 @@ export class IPC {
 			return OPERATION_SUCCESS
 
 		case 'unpause':
-			generalOpTarget.unpause()
+			generalOpTarget.unpause(query.who)
 			return OPERATION_SUCCESS
 
 		case 'retry':
@@ -361,7 +363,7 @@ export class IPC {
 				return {statusCode: 400, message: 'Invalid CL parameter: ' + cl}
 			}
 
-			let prevCl = generalOpTarget.forceSetLastClWithContext(cl, query.who, query.reason)
+			let prevCl = generalOpTarget.forceSetLastClWithContext(cl, query.who, query.reason, !!query.unblock && query.unblock === 'true')
 
 			this.ipcLogger.info(`Forcing last CL=${cl} on ${botname} : ${branch.name} (was CL ${prevCl}), ` +
 														`requested by ${query.who} (Reason: ${query.reason})`)
@@ -488,6 +490,13 @@ export class IPC {
 			}
 			this.ipcLogger.error('Error processing stomp: ' + operationResult.message)
 			return { statusCode: 500, message: operationResult.message }
+
+		case 'bypassgatewindow':
+			const sense = query.sense.toLowerCase().startsWith('t');
+			const prefix = sense ? 'en' : 'dis'
+			this.ipcLogger.info(`${query.who} ${prefix}abled gate window bypass on ${nodeName}->${edgeName}`)
+			edge!.bypassGateWindow(sense)
+			return { statusCode: 200, message: 'ok' }
 		}
 
 		return {statusCode: 404, message: 'Unrecognized node op: ' + operation}
diff --git a/Engine/Extras/RoboMerge/v3/src/robo/nodebot.ts b/Engine/Extras/RoboMerge/v3/src/robo/nodebot.ts
--- a/Engine/Extras/RoboMerge/v3/src/robo/nodebot.ts
+++ b/Engine/Extras/RoboMerge/v3/src/robo/nodebot.ts
@@ -5,12 +5,12 @@ import { _nextTick } from '../common/helper';
 import { ContextualLogger } from '../common/logger';
 import { Mailer, MailParams, Recipients } from '../common/mailer';
 import { Change, ConflictedResolveNFile, PerforceContext, RoboWorkspace } from '../common/perforce';
-import { IPCControls, NodeBotInterface } from './bot-interfaces';
-import { BlockagePauseInfo, BranchStatus, QueuedChange, ReconsiderArgs } from './status-types';
-import { AlreadyIntegrated, Blockage, Branch, BranchArg, BranchGraphInterface, ChangeInfo, Failure } from './branch-interfaces';
+import { IPCControls, NodeBotInterface, QueuedChange, ReconsiderArgs } from './bot-interfaces';
+import { BlockagePauseInfo, BranchStatus } from './status-types';
+import { AlreadyIntegrated, Blockage, Branch, BranchArg, BranchGraphInterface, ChangeInfo, EndIntegratingToGateEvent, Failure } from './branch-interfaces';
 import { MergeAction, OperationResult, PendingChange, resolveBranchArg, StompedRevision, StompVerification, StompVerificationFile }  from './branch-interfaces';
 import { Conflicts } from './conflicts';
-import { EdgeBot, EdgeIntegrationResult } from './edgebot';
+import { EdgeBot, EdgeMergeResults } from './edgebot';
 import { BotEventTriggers } from './events';
 import { PerforceStatefulBot } from './perforce-stateful-bot';
 import { BlockageNodeOpUrls, OperationUrlHelper } from './roboserver';
@@ -47,46 +47,6 @@ function isStompableNonBinary(filename: string) {
 
 type InfoResult = { info?: ChangeInfo, errors: any[] }
 
-type EdgeMerge = 'ok' | 'skipped' | 'error'
-
-/** results from integrating along available outgoing edges */
-class EdgeMergeResults {
-	// trying out making this purely per edge
-
-	addIntegrationResult(action: MergeAction, result: EdgeIntegrationResult) {
-		if ((result as Change).change) {
-			this.edges.set(action.branch.name, {merge: 'ok', change: result as Change})
-		}
-		else {
-			this.edges.set(action.branch.name, {merge: 'error'})
-			this.errors = [...this.errors, ...(result as string[])]
-		}
-	}
-
-	allChanges() {
-		return [...this.edges].map(([_, v]) => v.change).filter(Boolean) as Change[]
-	}
-
-	markSkipped(target: Branch) {
-		this.edges.set(target.name, {merge: 'skipped'})
-	}
-
-	wasSkipped(target: Branch) {
-		return this.edges.has(target.name) && this.edges.get(target.name)!.merge === 'skipped'
-	}
-
-	/** shelved or committed cl */
-	getTargetCl(target: Branch) {
-		const val = this.edges.get(target.name)
-		return val && val.change && val.change.change
-	}
-
-	// tuples of target name, result, opt change
-	edges = new Map<string, {merge: EdgeMerge, change?: Change}>()
-	// accumulated errors for all edges
-	errors: string[] = []
-}
-
 
 export class NodeBot extends PerforceStatefulBot implements NodeBotInterface {
 	readonly graphBotName: string
@@ -104,6 +64,7 @@ export class NodeBot extends PerforceStatefulBot implements NodeBotInterface {
 	protected readonly p4: PerforceContext
 
 	private readonly conflicts: Conflicts
+	private edgesQueuedToUnblock: Set<string>
 	private queuedChanges: QueuedChange[]
 
 	private readonly externalRobomergeUrl: string
@@ -124,10 +85,11 @@ export class NodeBot extends PerforceStatefulBot implements NodeBotInterface {
 		externalUrl: string, 
 		private readonly eventTriggers: BotEventTriggers,
 		settings: Context,
-		private readonly ubergraph: GraphAPI
+		private readonly ubergraph: GraphAPI,
+		private postIntegrate: () => Promise<void>
 	) {
 		// Node initial CL used when creating new (sub-)graph: should be able to set initial CL of source node and have that apply to edges
-		super(settings, branchDef.config.initialCL)
+		super(settings, reason => this.unblock(reason), branchDef.config.initialCL)
 		
 		this.branch = branchDef
 		this.branchGraph = branchDef.parent
@@ -142,6 +104,9 @@ export class NodeBot extends PerforceStatefulBot implements NodeBotInterface {
 
 		const settingsQueuedCls = this.settings.get('queuedCls')
 		this.queuedChanges = Array.isArray(settingsQueuedCls) ? settingsQueuedCls : []
+
+		const settingsQueuedUnblocks = this.settings.get('edgesQueuedToUnblock')
+		this.edgesQueuedToUnblock = new Set(Array.isArray(settingsQueuedUnblocks) ? settingsQueuedUnblocks : [])
 		
 		this.eventTriggers = eventTriggers
 		this.conflicts = new Conflicts(branchDef.upperName, eventTriggers, this.settings, externalUrl, this.branchGraph.config.reportToBuildHealth, this.nodeBotLogger)
@@ -213,15 +178,15 @@ export class NodeBot extends PerforceStatefulBot implements NodeBotInterface {
 	 * @throws {Error} If the branch argument does not have an immediate edge representing it for this NodeBot,
 	 *  this will throw an error.
 	 */
-	private getImmediateEdge(branch: BranchArg): EdgeBot | undefined {
+	getImmediateEdge(branch: BranchArg): EdgeBot | null {
 		let branchName = resolveBranchArg(branch, true)
-		return this.edges.get(branchName)
+		return this.edges.get(branchName) || null
 	}
 
-	getEdgeIPCControls(branch: BranchArg): IPCControls | undefined {
+	getEdgeIPCControls(branch: BranchArg): IPCControls | null {
 		const edge = this.getImmediateEdge(branch)
 
-		return edge ? edge.ipcControls : undefined  
+		return edge ? edge.ipcControls : null  
 	}
 
 	hasEdge(branch: BranchArg): boolean {
@@ -269,6 +234,7 @@ export class NodeBot extends PerforceStatefulBot implements NodeBotInterface {
 		return PerforceStatefulBot.setBotToLatestClInBranch(this, this.branch)
 	}
 
+	/** @return true if full tick completed, only used for analytics */
 	async tick() {
 		for (const edgeBot of this.edges.values()) {
 			await edgeBot.tick()
@@ -276,6 +242,21 @@ export class NodeBot extends PerforceStatefulBot implements NodeBotInterface {
 
 		// Pre-tick
 
+		if (this.edgesQueuedToUnblock) {
+			for (const edgeName of this.edgesQueuedToUnblock) {
+				const edge = this.getImmediateEdge(edgeName)
+				if (edge) {
+					edge.unblock('queued unblock')
+				}
+				else {
+					// could have been removed
+					this.nodeBotLogger.warn(`Couldn't find edge ${edgeName} to unblock`)
+				}
+			}
+			this.edgesQueuedToUnblock.clear()
+			this.persistQueuedChanges()
+		}
+
 		// Find the changelist # to start our p4 changes command at
 		let minCl = this.lastCl
 		// Gather the list of available edges for tick -- edges may change mid-tick but we will only allow this set to perform integrations
@@ -294,13 +275,6 @@ export class NodeBot extends PerforceStatefulBot implements NodeBotInterface {
 		}
 	// End Pre-tick
 
-		// allow processing manually queued changes
-		if (this.queuedChanges.length > 0) {
-			const fromQueue = this.queuedChanges.shift()!
-			await this.processQueuedChange(fromQueue)
-			return false
-		}
-
 		// see if our flow is paused
 		if (this.isManuallyPaused) {
 			this.nodeBotLogger.silly('tick() - This bot is manually paused')
@@ -355,17 +329,25 @@ export class NodeBot extends PerforceStatefulBot implements NodeBotInterface {
 		}
 		else {
 			this.ticksSinceLastNewP4Commit = 0
+
 			this.headCL = changes[0].change
 			await this._processListOfChanges(availableEdges, changes, MAX_CHANGES_TO_PROCESS_BEFORE_YIELDING)
 		}
 		return true
 	}
 
-	private async processQueuedChange(fromQueue: QueuedChange) {
+	queueEdgeUnblock(targetUpperName: string) {
+		this.edgesQueuedToUnblock.add(targetUpperName)
+	}
+
+	async processQueuedChange(fromQueue: QueuedChange) {
 		let logMessage = `Processing manually queued change ${fromQueue.cl} on ${this.fullName}, requested by ${fromQueue.who}`
 		if (fromQueue.workspace) {
 			logMessage += `, in workspace ${fromQueue.workspace}`
 		}
+		if (fromQueue.timestamp) {
+			logMessage += ` (delay: ${Math.round((Date.now() - fromQueue.timestamp) / 1000)})`
+		}
 		
 		let specifiedTargetBranch : Branch | null = null
 		if (fromQueue.targetBranchName) {
@@ -379,8 +361,6 @@ export class NodeBot extends PerforceStatefulBot implements NodeBotInterface {
 					`Manually queued change ${fromQueue.cl} on ${this.fullName}, requested by ${fromQueue.who}, ` + 
 						`specifies non-existant branch ${fromQueue.targetBranchName}. Aborting reconsider.`
 				)
-				this._persistQueuedChanges()
-				// Persist the remaining queued changes and carry on
 				return
 			}
 			else if (specifiedTargetBranch === this.branch) {
@@ -392,7 +372,6 @@ export class NodeBot extends PerforceStatefulBot implements NodeBotInterface {
 						`specifies integration from branch ${fromQueue.targetBranchName} into itself. Aborting reconsider.`
 				)
 				// Persist the remaining queued changes and carry on
-				this._persistQueuedChanges()
 				return
 			}
 
@@ -404,21 +383,20 @@ export class NodeBot extends PerforceStatefulBot implements NodeBotInterface {
 		// Get Change object from P4 for queued cl
 		const changeResult = await this._getChange(fromQueue.cl)
 
-		if (!changeResult.changes || changeResult.changes.length < 1) {
-			// change is invalid, so persist its removal from the queue
-			this._persistQueuedChanges()
+		if (typeof changeResult === 'string') {
+			// change is invalid
 
 			// @todo distinguish between Perforce errors and change not existing in branch
 
 
-			this.nodeBotLogger.error(`${this.fullName} - Error while querying P4 for change ${fromQueue.cl}: ${changeResult.errors}`)
+			this.nodeBotLogger.error(`${this.fullName} - Error while querying P4 for change ${fromQueue.cl}: ${changeResult}`)
 			// @todo email admin page user when users have to log in
 			// any way to show error without blocking page?
 			return
 		}
 
 		// prep the change
-		const change : Change = changeResult.changes[0]
+		const change = changeResult as Change
 		change.isUserRequest = true
 
 		const additionalFlags = fromQueue.additionalFlags || []
@@ -440,7 +418,6 @@ export class NodeBot extends PerforceStatefulBot implements NodeBotInterface {
 
 		// process just this change
 		await this._processAndMergeCl(new Map<string, EdgeBot>(this.edges), change, true, fromQueue.who, fromQueue.workspace, specifiedTargetBranch)
-		this._persistQueuedChanges()
 		return
 	}
 
@@ -457,8 +434,8 @@ export class NodeBot extends PerforceStatefulBot implements NodeBotInterface {
 		return this.nodeBotLogger
 	}
 
-	forceSetLastClWithContext(value: number, culprit: string, reason: string) {
-		const prevValue = this.forceSetLastCl(value)
+	forceSetLastClWithContext(value: number, culprit: string, reason: string, _?: boolean) {
+		const prevValue = this._forceSetLastCl_NoReset(value)
 
 		if (prevValue !== value) {
 			this.onForcedLastCl(this.displayName, value, prevValue, culprit, reason)
@@ -466,22 +443,13 @@ export class NodeBot extends PerforceStatefulBot implements NodeBotInterface {
 		return prevValue
 	}
 	
-	protected forceSetLastCl(value: number) {
-		const previousValue = super.forceSetLastCl(value)
-
-		// this also resets queued changes
-		this.queuedChanges = []
-		this._persistQueuedChanges()
-
-		return previousValue
-	}
-
 	onForcedLastCl(nodeOrEdgeName: string, forcedCl: number, previousCl: number, culprit: string, reason: string) {
 		this.conflicts.onForcedLastCl({nodeOrEdgeName, forcedCl, previousCl, culprit, reason})
 	}
 
-	protected _persistQueuedChanges() {
+	persistQueuedChanges() {
 		this.settings.set('queuedCls', this.queuedChanges);
+		this.settings.set('edgesQueuedToUnblock', [...this.edgesQueuedToUnblock]);
 	}
 
 	reconsider(instigator: string, changenum: number, additionalArgs?: Partial<ReconsiderArgs>) {
@@ -500,8 +468,12 @@ export class NodeBot extends PerforceStatefulBot implements NodeBotInterface {
 			}
 		}
 		this.nodeBotLogger.info(logMessage);
-		this.queuedChanges.push({cl: changenum, who: instigator, ...(additionalArgs || {})});
-		this._persistQueuedChanges();
+		this.queuedChanges.push({cl: changenum, who: instigator, timestamp: Date.now(), ...(additionalArgs || {})});
+		this.persistQueuedChanges();
+	}
+
+	popRequestedIntegration() {
+		return this.queuedChanges.shift()
 	}
 
 	createShelf(owner: string, workspace: string, changeCl: number, targetBranchName: string): OperationResult {
@@ -597,14 +569,14 @@ export class NodeBot extends PerforceStatefulBot implements NodeBotInterface {
 		// Get Change object from P4 for blockage cl
 		const blockageChangeRetrieval = await this._getChange(changeCl)
 
-		if (!blockageChangeRetrieval.changes || blockageChangeRetrieval.changes.length < 1) {
+		if (typeof blockageChangeRetrieval === 'string') {
 			return { 
 				success: false, 
-				message: "Unable to retrieve change infomation for given conflict. " + blockageChangeRetrieval.errors
+				message: `Unable to retrieve change infomation for given conflict. ${blockageChangeRetrieval}`
 			}
 		}
 
-		const blockageChange = blockageChangeRetrieval.changes[0]
+		const blockageChange = blockageChangeRetrieval as Change
 
 		// Set a variety of flags on the change before processing
 		blockageChange.isUserRequest = true
@@ -613,7 +585,7 @@ export class NodeBot extends PerforceStatefulBot implements NodeBotInterface {
 
 		// Process and attempt the integration
 		const edgeMap = new Map<string, EdgeBot>([[edge.targetBranch.upperName, edge]])
-		let mergeResult: EdgeMergeResults | undefined // Set to possible undefines as we use this in the finally clause
+		let integratedCl: Change | null = null // Set to possible undefines as we use this in the finally clause
 		try {
 			const processResult = await this._createChangeInfo(blockageChange, edgeMap, this.p4.username, targetBranch.workspace, targetBranch)
 			if (!processResult.info) {
@@ -622,7 +594,7 @@ export class NodeBot extends PerforceStatefulBot implements NodeBotInterface {
 					message: `Could not parse blockage change\nErrors: ${processResult.errors}`
 				}
 			}
-			mergeResult = await this._mergeClViaEdges(edgeMap, processResult.info, true)
+			const mergeResult = await this._mergeClViaEdges(edgeMap, processResult.info, true)
 
 			// If we don't have a change returned, we ran into an error processing the CL
 			const changes = mergeResult.allChanges()
@@ -637,19 +609,32 @@ export class NodeBot extends PerforceStatefulBot implements NodeBotInterface {
 				}
 			}
 
-			const integratedCl = changes[0]!
-			if (!integratedCl.change) {
+			const changelist = changes[0]!
+			if (changelist <= 0) {
 				return { 
 					success: false, 
-					message: `Was unable to get integration shelf for CL ${changeCl}: ${blockageChangeRetrieval.errors}`
+					message: `Unable to get integration shelf for CL ${changeCl}`
 				}
 			}
-			
+
+			let p4ChangeResult
+			try {
+				p4ChangeResult = await this.p4.getChange( targetBranch.rootPath, changelist, 'shelved')
+			}
+			catch (err) {
+				return { 
+					success: false, 
+					message: `Unable to get integration shelf for CL ${changeCl}: ${err}`
+				}
+			}
+
+			integratedCl = p4ChangeResult as Change
+
 			// Unshelve 
 			if (!await this.p4.unshelve(targetBranch.workspace, integratedCl.change)) {
 				return { 
 					success: false, 
-					message: `Was unable to get integration shelf for CL ${changeCl}: ${blockageChangeRetrieval.errors}`
+					message: `Unable to unshelve CL ${integratedCl.change}`
 				}
 			}
 
@@ -841,10 +826,8 @@ export class NodeBot extends PerforceStatefulBot implements NodeBotInterface {
 		finally {
 			// Delete the CL, including any remaining shelved files
 			this.nodeBotLogger.info(`${this.fullName} - Cleaning up after verifyStompRequest()`)
-			if (mergeResult) {
-				for (const processResultChange of mergeResult.allChanges()) {
-					await edge.revertPendingCLWithShelf(processResultChange.client, processResultChange.change, "Stomp Verify")
-				}
+			if (integratedCl) {
+				await edge.revertPendingCLWithShelf(integratedCl.client, integratedCl.change, 'Stomp Verify')
 			}
 		}
 	}
@@ -1089,6 +1072,8 @@ export class NodeBot extends PerforceStatefulBot implements NodeBotInterface {
 		status.conflicts = [] as any[]
 		this.conflicts.applyStatus(status.conflicts)
 
+		status.tick_count = this.tickCount
+
 		const edgeObj: { [key: string]: any } = {}
 		for (const [branchName, edgeBot] of this.edges) {
 			const edgeStatus = edgeBot.createStatus()
@@ -1211,6 +1196,16 @@ export class NodeBot extends PerforceStatefulBot implements NodeBotInterface {
 		this.conflicts.onAlreadyIntegrated(event)
 	}
 
+	onEndIntegratingToGate(event: EndIntegratingToGateEvent) {
+		if (event.context.from === this.branch) {
+			this.nodeBotLogger.info(`Gate closed info:
+	to: ${event.context.to.upperName}
+	edgeLastCl: ${event.context.edgeLastCl}
+	targetCl: ${event.targetCl}
+`)
+		}
+	}
+
 	// Listener method to email a blockage owner when the branch encounters said blockage,
 	// depending on NodeBot configuration
 	private emailBlockageCulprit(blockage: Blockage) {
@@ -1283,6 +1278,8 @@ export class NodeBot extends PerforceStatefulBot implements NodeBotInterface {
 		// make sure the list is sorted in ascending order
 		changes.sort((a, b) => a.change - b.change)
 
+		const integrationsPerEdge = new Map<string, number>()
+
 		for (let changeIndex = 0; changeIndex < changes.length; ++changeIndex) {
 			const change = changes[changeIndex]
 			const changeResult = await this._processAndMergeCl(availableEdges, change, false)
@@ -1299,15 +1296,22 @@ export class NodeBot extends PerforceStatefulBot implements NodeBotInterface {
 			// Every edge that was available at the start of the tick that is still available had the chance to receive this change, even if they weren't a target
 			// Count this as processed to keep them up-to-date
 			for (const edgeBot of availableEdges.values()) {
-				if (edgeBot.isAvailable && edgeBot.lastCl < change.change && !changeResult.wasSkipped(edgeBot.targetBranch)) {
-					edgeBot.updateLastCl(changes, changeIndex, changeResult.getTargetCl(edgeBot.targetBranch))
-				}
+				edgeBot.onNodeProcessedChange(changes, changeIndex, changeResult)
 			}
 
 			// yield if necessary - might be better to make this time-based
-			if (maxChangesToProcess > 0 && changeIndex === maxChangesToProcess) {
-				this.nodeBotLogger.info(`${this.branch.name} (${this.graphBotName}) yielding after ${maxChangesToProcess} revisions`)
-				return
+			if (maxChangesToProcess > 0) {
+
+				for (const [targetName, result] of changeResult.edges) {
+					if (result.result === 'ok') {
+						const numIntegrations = (integrationsPerEdge.get(targetName) || 0) + 1
+						if (numIntegrations >= maxChangesToProcess) {
+							this.nodeBotLogger.info(`${targetName} yielding after ${maxChangesToProcess} revisions`)
+							availableEdges.delete(targetName)
+						}
+						integrationsPerEdge.set(targetName, numIntegrations)
+					}
+				}
 			}
 
 			// If we've been paused (or blocked?) while the previous change was being processed, stop now
@@ -1551,15 +1555,15 @@ export class NodeBot extends PerforceStatefulBot implements NodeBotInterface {
 			change.desc = '<description not available>'
 		}
 
-		const parse = (lines: string[]) => parseDescriptionLines(
+		const parse = (lines: string[]) => parseDescriptionLines({
 			lines,
-			this.branch.isDefaultBot,
-			this.graphBotName,
-			change.change,
-			(this.branchGraph.config.alias || '').toUpperCase(),
-			this.branchGraph.config.macros,
-			this.nodeBotLogger
-		)
+			isDefaultBot: this.branch.isDefaultBot,
+			graphBotName: this.graphBotName,
+			cl: change.change,
+			aliasUpper: (this.branchGraph.config.alias || '').toUpperCase(),
+			macros: this.branchGraph.config.macros,
+			logger: this.nodeBotLogger
+		})
 
 		const parsedLines = parse(change.desc.split('\n'))
 
@@ -1635,9 +1639,63 @@ export class NodeBot extends PerforceStatefulBot implements NodeBotInterface {
 
 		const args = parsedLines.arguments
 		if (optTargetBranch || (args.length + defaultTargets.length) > 0) {
-			computeTargets(this.branch, this.ubergraph, change.change, info, args, defaultTargets, this.nodeBotLogger, !this.branch.config.disallowDeadend, optTargetBranch)
+			const computeResult = computeTargets(this.branch, this.ubergraph, change.change, info, args, defaultTargets, this.nodeBotLogger, optTargetBranch)
+			if (computeResult) {
+				const [computedMergeActions, flags] = computeResult
+
+				if (computedMergeActions.length > 0) {
+					if (info.userRequest) {
+						info.targets = computedMergeActions
+					}
+					else {
+						const mergeActions: MergeAction[] = []
+
+						for (const action of computedMergeActions) {
+							const edgeProperties = this.branch.edgeProperties.get(action.branch.upperName)
+
+							// let ignore and disregard cancel each other out here, but don't honour ignore here so we can validate it at the bottom of the function
+							const authorIsExcluded = edgeProperties && edgeProperties.excludeAuthors && edgeProperties.excludeAuthors.indexOf(info.author) >= 0
+
+							if (!authorIsExcluded || (flags.has('disregardexcludedauthors') && !flags.has('ignore'))) {
+								mergeActions.push(action)
+							}
+							else {
+								// skip this one
+								const skipMessage = `Skipping CL ${info.cl} due to excluded author '${info.author}'`
+								this.nodeBotLogger.info(skipMessage)
+
+							}
+						}
+
+						if (mergeActions.length > 0) {
+							info.targets = mergeActions
+						}
+					}
+
+					if (info.targets && flags.has('ignore')) {
+						if (this.branch.config.disallowDeadend) {
+							info.errors = ['deadend (ignore) disabled for this stream']
+						}
+						else {
+							delete info.targets
+						}
+					}
+				}
+			}
+
 		}
 
+// try doing all of disallowDeadend, exludedAuthor and disregaredExlcuded here
+//	disallowDeadend is by source node, others
+//	only complain about disallowDeadend if any targets remaining 
+// recreate target list, excluding 
+
+	// if (ri.flags.has('ignore')) {
+	// 	return { computeResult: null, errors: allowIgnore ? [] : ['deadend (ignore) disabled for this stream'] }
+	// }
+
+
+
 		if (!info.errors) {
 			const errors: string[] = []
 			const actions = info.targets || []
@@ -1674,14 +1732,6 @@ export class NodeBot extends PerforceStatefulBot implements NodeBotInterface {
 				continue
 			}
 
-			if (!info.userRequest && !target.flags.has('disregardexcludedauthors') && targetEdge.excludedAuthors.indexOf(info.author) >= 0) {
-				// skip this one
-				const skipMessage = `Skipping CL ${info.cl} due to excluded author '${info.author}'`
-				this.nodeBotLogger.info(skipMessage)
-
-				continue
-			}
-
 			// Gate problem: say gate is 15, last integrated on edge was 12, incoming is 17. Edge is available because 12 < 15
 			// Note: good thing about using availability (in theory) - will not fetch loads of revisions repeatedly
 
@@ -1709,7 +1759,12 @@ export class NodeBot extends PerforceStatefulBot implements NodeBotInterface {
 				++this.tickJournal.merges
 			}
 
-			const integrationResult = await targetEdge.performMerge(info, target, this.branch.convertIntegratesToEdits)
+			const additionalDescription = 
+				targetEdge.implicitCommands
+					.map(s => s + '\n')
+					.join('')
+
+			const integrationResult = await targetEdge.performMerge(info, target, this.branch.convertIntegratesToEdits, additionalDescription)
 
 			// wake up target nodebot, so look for further merges or conflict resolutions
 			const targetNodebot = target.branch.bot as NodeBot
@@ -1725,6 +1780,11 @@ export class NodeBot extends PerforceStatefulBot implements NodeBotInterface {
 			}
 
 			mergedChanges.addIntegrationResult(target, integrationResult)
+
+			// if integration did something, enough time has probably passed to make worth checking to see if any manual requests came in
+			if (integrationResult.result !== 'nothing to do') {
+				await this.postIntegrate()
+			}
 		}
 
 		return mergedChanges
@@ -1768,6 +1828,7 @@ export class NodeBot extends PerforceStatefulBot implements NodeBotInterface {
 		}
 	}
 
+	/** return true if this is a new flow blockage (excludes reconsiders) */
 	handleMergeFailure(failure: Failure, pending: PendingChange, postIntegrate?: boolean) {
 		const targetBranch = pending.action.branch
 		const shortMessage = `${failure.kind} while merging CL ${pending.change.cl} to ${targetBranch.name}`
@@ -1803,8 +1864,10 @@ export class NodeBot extends PerforceStatefulBot implements NodeBotInterface {
 				}
 
 				this.conflicts.onBlockage(blockage)
+				return true
 			}
 		}
+		return false
 	}
 
 	private sendNagEmails() {
diff --git a/Engine/Extras/RoboMerge/v3/src/robo/notifications.ts b/Engine/Extras/RoboMerge/v3/src/robo/notifications.ts
--- a/Engine/Extras/RoboMerge/v3/src/robo/notifications.ts
+++ b/Engine/Extras/RoboMerge/v3/src/robo/notifications.ts
@@ -1,6 +1,7 @@
 // Copyright Epic Games, Inc. All Rights Reserved.
 
 import * as fs from 'fs';
+import { DateTimeFormatOptions } from 'intl';
 import { Args } from '../common/args';
 import { Badge } from '../common/badge';
 import { Random } from '../common/helper';
@@ -25,6 +26,7 @@ export const SLACK_MESSAGES_PERSISTENCE_KEY = 'slackMessages'
 const CHANGELIST_FIELD_TITLE = 'Change'
 const ACKNOWLEDGED_FIELD_TITLE = 'Acknowledged'
 const SUNSET_PERSISTED_NOTIFICATIONS_DAYS = 30
+const EPIC_TIME_OPTIONS: DateTimeFormatOptions = {timeZone: 'EST5EDT', timeZoneName: 'short'}
 
 const KNOWN_BOT_NAMES = ['buildmachine', 'robomerge'];
 
@@ -454,6 +456,7 @@ export class BotNotifications implements BotEventHandler {
 		if (DIRECT_MESSAGING_ENABLED && !isBotUser && targetBranch && userEmail) {
 			const dmText = `Your change (${makeClLink(changeInfo.source_cl)}) ` +
 				`hit '${issue}' while merging from *${sourceBranch.name}* to *${targetBranch.name}*.\n\n` +
+				'`' + blockage.change.description.substr(0, 80) + '`\n\n' +
 				"*_Resolving this blockage is time sensitive._ Please select one of the following:*"
 			
 			const urls = this.blockageUrlGenerator(blockage)
@@ -473,7 +476,9 @@ export class BotNotifications implements BotEventHandler {
 			const targetKey = info.targetBranchName || info.kind
 			const title = ACKNOWLEDGED_FIELD_TITLE
 			if (info.acknowledger) {
-				this.tryAddFieldToChannelMessages(info.sourceCl, targetKey, {title, value: info.acknowledger, short: true})
+				// hard code to Epic (!) Standard/Daylight Time
+				const suffix = info.acknowledgedAt ? ' at ' + info.acknowledgedAt.toLocaleTimeString('en-US', EPIC_TIME_OPTIONS) : ''
+				this.tryAddFieldToChannelMessages(info.sourceCl, targetKey, {title, value: info.acknowledger + suffix, short: true})
 			}
 			else {
 				this.tryRemoveFieldFromChannelMessages(info.sourceCl, targetKey, title)
diff --git a/Engine/Extras/RoboMerge/v3/src/robo/perforce-stateful-bot.ts b/Engine/Extras/RoboMerge/v3/src/robo/perforce-stateful-bot.ts
--- a/Engine/Extras/RoboMerge/v3/src/robo/perforce-stateful-bot.ts
+++ b/Engine/Extras/RoboMerge/v3/src/robo/perforce-stateful-bot.ts
@@ -1,17 +1,12 @@
 // Copyright Epic Games, Inc. All Rights Reserved.
 import { ContextualLogger, NpmLogLevel } from "../common/logger";
 import { Change, ChangelistStatus, PerforceContext } from "../common/perforce";
-import { BotIPC } from "./bot-interfaces";
-import { BlockagePauseInfo, BlockagePauseInfoMinimal, ReconsiderArgs } from "./status-types";
+import { BotIPC, ReconsiderArgs } from "./bot-interfaces";
+import { BlockagePauseInfo, BlockagePauseInfoMinimal } from "./status-types";
 import { Branch, OperationResult } from "./branch-interfaces";
 import { Context } from "./settings";
 import { PauseState } from "./state-interfaces";
 
-export interface PerforceRequestResult {
-	changes?: Change[] | null
-	errors: any[]
-}
-
 export abstract class PerforceStatefulBot implements BotIPC {
 	protected p4: PerforceContext
 	abstract branch: Branch
@@ -28,8 +23,8 @@ export abstract class PerforceStatefulBot implements BotIPC {
 	isRunning = false
 	readonly pauseState: PauseState
 
-	constructor(protected readonly settings: Context, initialCl?: number) {
-		this.pauseState = new PauseState((reason: string) => { this.unblock(reason) }, settings)
+	constructor(protected readonly settings: Context, unblockFunc: (reason: string) => void, initialCl?: number) {
+		this.pauseState = new PauseState(unblockFunc, settings)
 		this._lastClProcessed = this.settings.getInt('lastCl', initialCl || 0)
 
 		if (this.pauseState.isBlocked()) {
@@ -82,13 +77,13 @@ export abstract class PerforceStatefulBot implements BotIPC {
 		this.pauseState.manuallyPause(message, owner)
 		this.logger.info(`Paused ${this.fullName}: ${message}.`)
 	}
-	unpause() {
+	unpause(requester: string) {
 		if (!this.isManuallyPaused) {
 			return
 		}
 
 		this.pauseState.unpause()
-		this.logger.info(`Unpaused ${this.fullName}.`)
+		this.logger.info(`Unpaused ${this.fullName}, requested by ${requester}`)
 	}
 
 	abstract reconsider(instigator: string, changeCl: number, additionalArgs?: Partial<ReconsiderArgs>): void
@@ -128,14 +123,6 @@ export abstract class PerforceStatefulBot implements BotIPC {
 		// set the cache and save it
 		this._forceSetLastCl_NoReset(value);
 	}
-	protected forceSetLastCl(value: number) {
-		const prevValue = this._forceSetLastCl_NoReset(value)
-
-		// When we force set the last CL, we want to unblock
-		this.unblock(`${this.fullNameForLogging} last CL forcibly set to ${value}`)
-
-		return prevValue
-	}
 	protected _forceSetLastCl_NoReset(value: number): number {
 		if (value > this.lastBlockage) {
 			this.lastBlockage = -1
@@ -164,19 +151,15 @@ export abstract class PerforceStatefulBot implements BotIPC {
 		bot.lastCl = change.change;
 	}
 
-	abstract forceSetLastClWithContext(value: number, culprit: string, reason: string): number;
+	abstract forceSetLastClWithContext(value: number, culprit: string, reason: string, unblock: boolean): number;
 
-	async _getChange(changeCl: number, path?: string, status?: ChangelistStatus) : Promise<PerforceRequestResult> {
-		let change: Change | null = null
-		let errors: any[] = []
+	async _getChange(changeCl: number, path?: string, status?: ChangelistStatus) : Promise<Change | string> {
 		try {
-			change = await this.p4.getChange(path || this.branch.rootPath, changeCl, status)
+			return await this.p4.getChange(path || this.branch.rootPath, changeCl, status)
 		}
 		catch (err) {
-			return { errors: [err] }
+			return err.toString()
 		}
-
-		return { changes: [change], errors }
 	}
 
 	protected _log_action(action: string, logLevel: NpmLogLevel = 'info') {
@@ -210,4 +193,5 @@ export abstract class PerforceStatefulBot implements BotIPC {
 	abstract setBotToLatestCl(): Promise<void>
 
 	abstract tick(): Promise<boolean>;
+	tickCount = 0
 }
diff --git a/Engine/Extras/RoboMerge/v3/src/robo/robo.ts b/Engine/Extras/RoboMerge/v3/src/robo/robo.ts
--- a/Engine/Extras/RoboMerge/v3/src/robo/robo.ts
+++ b/Engine/Extras/RoboMerge/v3/src/robo/robo.ts
@@ -35,6 +35,9 @@ const roboStartupLogger = new ContextualLogger('Robo Startup')
 VersionReader.init(roboStartupLogger)
 
 
+const USER_WORKSPACE_EXCLUDE_PATTERNS: (RegExp | string)[] = [
+	'horde-p4bridge-hordeserver-'
+]
 
 // I seem to have broken this
 const DEBUG_SKIP_BRANCH_SETUP = false;
@@ -217,7 +220,8 @@ export class RoboMerge {
 	}
 
 	async getWorkspacesForUser(user: string) {
-		return await this.p4.find_workspaces(user)
+		return (await this.p4.find_workspaces(user))
+			.filter(ws => !USER_WORKSPACE_EXCLUDE_PATTERNS.some(entry => ws.client.match(entry)))
 	}
 
 	stop() {
@@ -231,7 +235,8 @@ async function _getExistingWorkspaces(user?: string) {
 }
 
 async function _initWorkspacesForGraphBot(graphBot: GraphBot, existingWorkspaces: Set<string>, logger: ContextualLogger) {
-	const workspacesToReset: string[] = []
+	// name and depot root pair
+	const workspacesToReset: [string, string][] = []
 
 	for (const branch of graphBot.branchGraph.branches) {
 		if (branch.workspace !== null) {
@@ -258,7 +263,7 @@ async function _initWorkspacesForGraphBot(graphBot: GraphBot, existingWorkspaces
 
 		// see if we already have this workspace
 		if (existingWorkspaces.has(ws)) {
-			workspacesToReset.push(ws);
+			workspacesToReset.push([ws, branch.rootPath])
 		}
 		else {
 			const params: any = {};
diff --git a/Engine/Extras/RoboMerge/v3/src/robo/roboserver.ts b/Engine/Extras/RoboMerge/v3/src/robo/roboserver.ts
--- a/Engine/Extras/RoboMerge/v3/src/robo/roboserver.ts
+++ b/Engine/Extras/RoboMerge/v3/src/robo/roboserver.ts
@@ -172,6 +172,12 @@ class RoboWebApp implements AppInterface {
 		return this.preview(clStr)
 	}
 
+	@Handler('GET', '/preview.html', {filetype: 'text/html'}) 
+	previewNoBotHtml() {
+		const clStr = this.request.url.searchParams.get('cl')
+		return clStr ? this.preview(clStr) : {statusCode: 400, message: "Query argument cl required"}
+	}
+
 	@Handler('GET', '/preview/*/*', {filetype: 'text/html'}) 
 	async preview(clStr: string, singleBot?: string) {
 
@@ -454,7 +460,7 @@ class RoboWebApp implements AppInterface {
 		return (await this.sendMessage('getp4tasks')).data || []
 	}
 
-	@SecureHandler('GET', '/api/crashserver', {requiredTags: ['admin']})
+	@SecureHandler('POST', '/api/crashserver', {requiredTags: ['admin']})
 	async crashServer() {
 		if (!this.authData) {
 			throw new Error('Secure call but no auth data?')
@@ -466,7 +472,7 @@ class RoboWebApp implements AppInterface {
 		throw new Error(msg)
 	}
 
-	@SecureHandler('GET', '/api/crashapi', {requiredTags: ['admin']})
+	@SecureHandler('POST', '/api/crashapi', {requiredTags: ['admin']})
 	async crashAPI() {
 		if (!this.authData) {
 			throw new Error('Secure call but no auth data?')
@@ -476,7 +482,7 @@ class RoboWebApp implements AppInterface {
 		return this.sendMessage('crashAPI', [this.authData.user])
 	}
 
-	@SecureHandler('GET', '/api/bot/*/crashgraphBot', {requiredTags: ['admin']})
+	@SecureHandler('POST', '/api/bot/*/crashgraphBot', {requiredTags: ['admin']})
 	async crashGraphBot(botname: string) {
 		if (!this.authData) {
 			throw new Error('Secure call but no auth data?')
diff --git a/Engine/Extras/RoboMerge/v3/src/robo/status-types.ts b/Engine/Extras/RoboMerge/v3/src/robo/status-types.ts
--- a/Engine/Extras/RoboMerge/v3/src/robo/status-types.ts
+++ b/Engine/Extras/RoboMerge/v3/src/robo/status-types.ts
@@ -1,4 +1,20 @@
-import { NodeOptions } from './branchdefs';
+// Copyright Epic Games, Inc. All Rights Reserved.
+
+type ReconsiderArgs = {
+	additionalFlags: string[]
+	workspace: string
+	targetBranchName: string
+	description: string
+	commandOverride: string
+}
+
+type QueuedChange = Partial<ReconsiderArgs> & {
+	cl: number
+	who: string
+	timestamp: number 
+}
+
+
 
 export type FailureKind =
 	'Integration error' |
@@ -107,6 +123,8 @@ export type NodeStatusFields = BotStatusFields & {
 
 	conflicts: ConflictStatusFields[]
 	edges: { [key: string]: EdgeStatusFields }
+
+	tick_count: number
 }
 
 // only includes stuff directly used by web code
@@ -125,7 +143,7 @@ export type BranchDefForStatus = {
 	aliases: any
 	blockAssetTargets: string[]
 	bot: string
-	config: NodeOptions
+	config: any //NodeOptions
 	convertIntegratesToEdits: boolean
 	defaultFlow: string[]
 	flowsTo: string[]
@@ -169,28 +187,3 @@ export type UserStatusData = {
 	botStates: [string, GraphBotState][]
 	insufficientPrivelege?: boolean
 }
-
-
-// queued change vs reconsider:
-
-//  reconsider is the name of user facing 'queue change' operation, but also the top level call to add to the queue,
-//  used by stomp etc.
-
-// slightly sneaky that since NodeBot.reconsider's additional args are optional, it fulfils both of these, i.e.
-// implements the base bot reconsider function, and it also the 'queue change' function.
-
-// ReconsiderArgs are the additional args used by stomp etc. I'm adding commandOverride to them, even thoguh that's a
-// slightly different pattern (unless it breaks something)
-
-export type ReconsiderArgs = {
-	additionalFlags: string[]
-	workspace: string
-	targetBranchName: string
-	description: string
-	commandOverride: string
-}
-
-export type QueuedChange = Partial<ReconsiderArgs> & {
-	cl: number
-	who: string
-}
diff --git a/Engine/Extras/RoboMerge/v3/src/robo/targets.ts b/Engine/Extras/RoboMerge/v3/src/robo/targets.ts
--- a/Engine/Extras/RoboMerge/v3/src/robo/targets.ts
+++ b/Engine/Extras/RoboMerge/v3/src/robo/targets.ts
@@ -65,6 +65,16 @@ const NEUTER_COMMANDS = [
 	"REVIEW",
 ]
 
+const ROBO_TAGS = [
+	'ROBOMERGE-AUTHOR',
+	'ROBOMERGE-BOT',
+	'ROBOMERGE-COMMAND',
+	'ROBOMERGE-CONFLICT',
+	'ROBOMERGE-EDIGRATE',
+	'ROBOMERGE-OWNER',
+	'ROBOMERGE-SOURCE',
+]
+
 type ChangeFlag = 'review' | 'manual' | 'null' | 'ignore' | 'disregardexcludedauthors'
 
 // mapping of #ROBOMERGE: flags to canonical names
@@ -99,6 +109,60 @@ function parseTargetList(targetString: string) {
 	return targetString.split(/[ ,]/).filter(Boolean)
 }
 
+
+type Token_Command = {
+	bot: string
+	param: string
+}
+
+type Token_Other = {
+	initialWhitespace: string
+	tag: string
+	rest: string
+
+	isRobo: boolean
+}
+
+// don't think this supports bare robomerge command with no args to disable
+//   default flow
+
+// result is whether we found and ignored a bare robomerge command
+function tokenizeLine(line: string): Token_Command[] | Token_Other | null {
+
+	// trim end - keep any initial whitespace
+	const comp = line.replace(/\s+$/, '')
+
+	// check for control hashes
+	const match = comp.match(/^(\s*)#([-\w[\]]+)[:\s]*(.*)$/)
+	if (!match) {
+		return null
+	}
+
+	const initialWhitespace = match[1] 
+	const tag = match[2].toUpperCase()
+	const rest = match[3].trim()
+
+	// check for any non-robomerge # tags we care about
+	if (initialWhitespace || !tag.startsWith('ROBOMERGE')) {
+		return {initialWhitespace, tag, rest, isRobo: false}
+	}
+
+	const bareCommand = tag === 'ROBOMERGE'
+	let botName = ''
+	if (!bareCommand)
+	{
+		const specificBotMatch = tag.match(/ROBOMERGE\[([-_a-z0-9]+)\]/i)
+		if (specificBotMatch) {
+			botName = specificBotMatch[1].toUpperCase()
+		}
+	}
+	if (!bareCommand && !botName) {
+		return {initialWhitespace, tag, rest, isRobo: true}
+	}
+
+	return parseTargetList(rest).map(param => ({bot: botName, param}))
+}
+
 export class DescriptionParser {
 	source = ''
 	owner: string | null = null
@@ -114,10 +178,9 @@ export class DescriptionParser {
 
 	// per target bot arguments
 	arguments: string[] = []
-	otherBotArguments: [string, string][] = []
+	otherBotArguments: Token_Command[] = []
 
 	expandedMacros: string[] = []
-	expandedMacroLines: string[] = []
 
 	errors: string[] = []
 
@@ -142,157 +205,136 @@ export class DescriptionParser {
 		this.errors = other.errors
 	}
 
-	parseLine(line: string) {
-
-		// trim end - keep any initial whitespace
-		const comp = line.replace(/\s+$/, '')
-
-		// check for control hashes
-		const match = comp.match(/^(\s*)#([-\w[\]]+)[:\s]*(.*)$/)
-		if (!match) {
-			// strip beginning blanks
-			if (this.descFinal.length > 0 || comp !== '') {
-				if (comp.indexOf('[NULL MERGE]') !== -1) {
-					this.propagatingNullMerge = true
+	private processNonRoboToken(token: Token_Other) {
+		if (token.isRobo) {
+			if (ROBO_TAGS.indexOf(token.tag) >= 0) {
+				if (token.tag === 'ROBOMERGE-AUTHOR') {
+					// tag intended to be purely for propagating
+					this.authorTag = token.rest
 				}
-				this.descFinal.push(line)
-			}
-			return
-		}
-
-		const ws = match[1] 
-		const command = match[2].toUpperCase()
-		const value = match[3].trim()
-
-		// #robomerge tags are required to be at the start of the line
-		if (ws && command === 'ROBOMERGE' && value.match(/\bnone\b/i)) {
-			// completely skip #robomerge nones that might come in from GraphBot branch integrations, so we don't
-			// get impaled on our own commit hook
-			return
-		}
-
-		// check for any non-robomerge # tags we care about
-		if (ws || !command.startsWith('ROBOMERGE')) {
-			// check for commands to neuter
-			if (NEUTER_COMMANDS.indexOf(command) >= 0) {
-				// neuter codereviews so they don't spam out again
-				this.descFinal.push(`${ws}[${command}] ${value}`)
-			}
-			
-			else if (command.indexOf('REVIEW-') === 0) {
-				// remove swarm review numbers entirely (they can't be neutered and generate emails)
-				if (value !== '') {
-					this.descFinal.push(value)
+				else if (token.tag === 'ROBOMERGE-OWNER') {
+					this.owner = token.rest.toLowerCase()
 				}
-			}
-			else if (command === 'OKFORGITHUB') {
-				this.hasOkForGithubTag = true
-				this.descFinal.push(line)
+				else if (token.tag === 'ROBOMERGE-SOURCE') {
+					this.source = `${token.rest} via CL ${this.cl}`
+				}
+				else if (token.tag === 'ROBOMERGE-BOT') {
+					// if matches description written in _mergeCl
+					if (token.rest.startsWith(this.graphBotName + ' (')) {
+						this.useDefaultFlow = false
+					}
 				}
-			else {
-
-				// default behavior is to keep any hashes we don't recognize
-				this.descFinal.push(line)
-			}
-			return
-		}
-
-		// Handle commands
-		if (command === 'ROBOMERGE') {
-
-			// ignore bare ROBOMERGE tags if we're not the default bot (should not affect default flow)
-			if (this.isDefaultBot) {
-				this.useDefaultFlow = false
-				this.processTargetList(value)
-			}
-			return
-		}
-		
-		// Look for a specified bot and other defined tags
-		const specificBotMatch = command.match(/ROBOMERGE\[([-_a-z0-9]+)\]/i)
-
-		if (specificBotMatch) {
-			const specificBot = specificBotMatch[1].toUpperCase()
-			if (specificBot === this.graphBotName ||
-				specificBot === this.aliasUpper ||
-				specificBot === 'ALL') {
-				this.useDefaultFlow = false
-
-// @todo 'ALL' with targets should be an error (but don't know at this point ...)
-				this.processTargetList(value)
 			}
 			else {
-				// keep a record of commands to forward to other bots - processed later
-				this.otherBotArguments = [
-					...this.otherBotArguments,
-					...(parseTargetList(value).map(arg => [specificBot, arg]) as [string, string][])
-				]
+				// add syntax error for unknown command
+				this.errors.push(`Unknown command '${token.tag}`)
 			}
+			return false
 		}
-		else if (command === 'ROBOMERGE-AUTHOR') {
-			// tag intended to be purely for propagating
-			// currently using it to set owner if not already set
-			this.authorTag = value
-		}
-		else if (command === 'ROBOMERGE-OWNER') {
-			this.owner = value.toLowerCase()
-		}
-		else if (command === 'ROBOMERGE-SOURCE') {
-			this.source = `${value} via CL ${this.cl}`
+
+		// check for commands to neuter
+		if (NEUTER_COMMANDS.indexOf(token.tag) >= 0) {
+			// neuter codereviews so they don't spam out again
+			this.descFinal.push(`${token.initialWhitespace}[${token.tag}] ${token.rest}`)
+			return false
 		}
-		else if (command === 'ROBOMERGE-BOT') {
-			// if matches description written in _mergeCl
-			if (value.startsWith(this.graphBotName + ' (')) {
-				this.useDefaultFlow = false
+		if (token.tag.match(/REVIEW-(\d+)/)) {
+							// remove swarm review numbers entirely (they can't be neutered and generate emails)
+			if (token.rest) {
+				this.descFinal.push(token.rest)
 			}
+			return false
 		}
-		else if (command !== 'ROBOMERGE-EDIGRATE' &&
-				 command !== 'ROBOMERGE-COMMAND' && 
-				 command !== 'ROBOMERGE-CONFLICT') {
-			// add syntax error for unknown command
-			this.errors.push(`Unknown command '${command}`)
+
+		if (token.tag === 'OKFORGITHUB') {
+			this.hasOkForGithubTag = true
 		}
+		return true
 	}
 
-	private processTargetList(targetString: string) {
-		const targetNames = parseTargetList(targetString)
+	parse(lines: string[]) {
+		for (const line of lines) {
+			const result = tokenizeLine(line)
+			if (!result) {
+				// strip beginning blanks
+				if (this.descFinal.length > 0 || line.trim()) {
+					if (line.indexOf('[NULL MERGE]') !== -1) {
+						this.propagatingNullMerge = true
+					}
+					this.descFinal.push(line)
+				}
+				continue
+			}
 
-		for (const name of targetNames) {
-			const macroLines = this.macros[name.toLowerCase()]
-			if (macroLines) {
-				this.expandedMacros.push(name)
-				this.expandedMacroLines = [...this.expandedMacroLines, ...macroLines]
+			const otherResult = result as Token_Other
+			if (otherResult.tag) {
+				if (this.processNonRoboToken(otherResult)) {
+					// default behavior is to keep any hashes we don't recognize
+					this.descFinal.push(line)
+				}
+				continue
 			}
-			else {
-				this.arguments.push(name)
+
+			const thisBotNames = ['', this.graphBotName, this.aliasUpper, 'ALL']
+			const commands: Token_Command[] = []
+			for (const tokenCommand of (result as Token_Command[])) {
+				const macroLines = this.macros[tokenCommand.param.toLowerCase()]
+				if (macroLines) {
+
+					// @todo: expand macros when loading config (with limited validation for preview),
+					// can be recursive
+					// for now, run parse line
+					for (const macro of macroLines) {
+						const macroResult = tokenizeLine(macro) as Token_Command[]
+						for (const expanded of macroResult) {
+							if (expanded.bot === undefined) {
+								this.errors.push(`Invalid macro '${tokenCommand.param}'`)
+								return
+							}
+							commands.push(expanded)
+						}
+					}
+					this.expandedMacros.push(tokenCommand.param)
+				}
+				else {
+					commands.push(tokenCommand)
+				}
+			}
+
+			for (const command of commands) {
+				if (thisBotNames.indexOf(command.bot) < 0) {
+					// other bot
+					this.otherBotArguments.push(command)
+					continue
+				}
+
+				// ignore bare ROBOMERGE tags if we're not the default bot (should not affect default flow)
+				if (!command.bot) {
+					if (this.isDefaultBot) {
+						this.useDefaultFlow = false
+					}
+					else {
+						continue
+					}
+				}
+				this.arguments.push(command.param)
 			}
 		}
 	}
 }
 
-export function parseDescriptionLines(
-	lines: string[],
-	isDefaultBot: boolean, 
-	graphBotName: string, 
-	cl: number, 
-	aliasUpper: string,
-	macros: {[name: string]: string[]},
-	logger: ContextualLogger
-	) {
-	const lineParser = new DescriptionParser(isDefaultBot, graphBotName, cl, aliasUpper, macros)
-	let safety = 5
-	while (lines.length > 0) {
-		for (const line of lines) {
-			lineParser.parseLine(line)
-		}
-		lines = lineParser.expandedMacroLines
-		lineParser.expandedMacroLines = []
-
-		if (--safety === 0) {
-			logger.warn('Got stuck in a loop expanding macros in CL#' + cl)
-			break
-		}
+export function parseDescriptionLines(args: {
+		lines: string[],
+		isDefaultBot: boolean, 
+		graphBotName: string, 
+		cl: number, 
+		aliasUpper: string,
+		macros: {[name: string]: string[]},
+		logger: ContextualLogger
 	}
+	) {
+	const lineParser = new DescriptionParser(args.isDefaultBot, args.graphBotName, args.cl, args.aliasUpper, args.macros)
+	lineParser.parse(args.lines)
 	return lineParser
 }
 
@@ -323,7 +365,7 @@ export function processOtherBotTargets(
 	parsedLines: DescriptionParser,
 	sourceBranch: Branch,
 	ubergraph: GraphAPI,
-	actions: MergeAction[],
+	actions: Readonly<Readonly<MergeAction>[]>,
 	errors: string[]
 	) {
 
@@ -356,15 +398,16 @@ export function processOtherBotTargets(
 	// calculate skip set first time we need it
 	let skipNodes: Node[] | null = null
 
-	for (const [bot, arg] of parsedLines.otherBotArguments) {
-		const targetBranchGraph = ubergraph.getBranchGraph(bot.toUpperCase() as BotName)
+	for (const token of parsedLines.otherBotArguments) {
+		const targetBranchGraph = ubergraph.getBranchGraph(token.bot.toUpperCase() as BotName)
 		if (!targetBranchGraph) {
-			errors.push(`Bot '${bot}' not found in ubergraph`)
+			errors.push(`Bot '${token.bot}' not found in ubergraph`)
 			continue
 		}
 
 		const targetBotName = targetBranchGraph.botname as BotName
 
+		const arg = token.param
 		if (isFlag(arg)) {
 			// no further checking for flags
 			botInfo(targetBranchGraph).commands.push(arg)
@@ -376,7 +419,7 @@ export function processOtherBotTargets(
 		if (!branch) {
 			if (!targetBranchGraph.config.macros[target.toLowerCase()] &&
 				targetBranchGraph.config.branchNamesToIgnore.indexOf(target.toUpperCase()) < 0) {
-				errors.push(`Branch '${target}' not found in ${bot}`)
+				errors.push(`Branch '${target}' not found in ${token.bot}`)
 			}
 			continue
 		}
@@ -464,15 +507,7 @@ class RequestedIntegrations {
 	readonly flags = new Set<ChangeFlag>()
 	integrations: [string, MergeMode][] = []
 
-	parse(commandArguments: string[], cl: number, forceStomp: boolean, logger: ContextualLogger) {
-		// If forceStompChanges, add the target as a stomp if it isn't a flagged target
-		if (forceStomp) {
-			this.integrations = commandArguments
-				.filter(arg => '!-$#'.indexOf(arg[0]) < 0)
-				.map(arg => [arg, 'clobber'])
-			return
-		}
-
+	parse(commandArguments: Readonly<string[]>, cl: number, logger: ContextualLogger) {
 		for (const arg of commandArguments) {
 			const argLower = arg.toLowerCase()
 
@@ -526,26 +561,19 @@ type ComputeTargetsResult =
 	, errors: string[] 
 	}
 
-// public so it can be called from unit tests
-export function computeTargetsImpl(
+function computeTargetsImpl(
 	sourceBranch: Branch, 
 	ubergraph: GraphAPI,
 	cl: number,
-	forceStomp: boolean,
-	allowIgnore: boolean,
-	commandArguments: string[],
-	defaultTargets: string[], 
+	commandArguments: Readonly<string[]>,
+	defaultTargets: Readonly<string[]>, 
 	logger: ContextualLogger
 ): ComputeTargetsResult {
 
 	// list of branch names and merge mode (default targets first, so merge mode gets overwritten if added explicitly)
 
 	const ri = new RequestedIntegrations
-	ri.parse(commandArguments, cl, forceStomp, logger)
-
-	if (ri.flags.has('ignore')) {
-		return { computeResult: null, errors: allowIgnore ? [] : ['deadend (ignore) disabled for this stream'] }
-	}
+	ri.parse(commandArguments, cl, logger)
 
 	if (ri.flags.has('null')) {
 		for (const branchName of sourceBranch.forceFlowTo) {
@@ -553,11 +581,9 @@ export function computeTargetsImpl(
 		}
 	}
 
-	const defaultMergeMode: MergeMode = forceStomp ? 'clobber' : 'normal'
-
 	const branchGraph = sourceBranch.parent
 	const requestedMerges = [
-		...defaultTargets.map(name => [name, defaultMergeMode] as [string, MergeMode]),
+		...defaultTargets.map(name => [name, 'normal'] as [string, MergeMode]),
 		...ri.integrations.filter(([name, _]) => branchGraph.config.branchNamesToIgnore.indexOf(name.toUpperCase()) < 0)
 	]
 
@@ -649,34 +675,30 @@ export function computeTargetsImpl(
 
 
 export function computeTargets(
-	sourceBranch: Branch,
+	sourceBranch: Readonly<Branch>,
 	ubergraph: GraphAPI,
 	cl: number,
 	info: TargetInfo,
-	commandArguments: string[],
-	defaultTargets: string[],
+	commandArguments: Readonly<string[]>,
+	defaultTargets: Readonly<string[]>,
 	logger: ContextualLogger,
-	allowIgnore: boolean,
-
 	optTargetBranch?: Branch | null
-) {
-
-
-	const { computeResult, errors } = computeTargetsImpl(sourceBranch, ubergraph, cl, info.forceStompChanges, allowIgnore, commandArguments, defaultTargets, logger)
+) : [MergeAction[], Set<ChangeFlag>] | null {
+	const { computeResult, errors } = computeTargetsImpl(sourceBranch, ubergraph, cl, commandArguments, defaultTargets, logger)
 
 	if (errors.length > 0) {
 		info.errors = errors
-		return
+		return null
 	}
 
 	if (!computeResult) {
-		return
+		return null
 	}
 
 	let { merges, flags, targets } = computeResult
 	if (!merges) {
 		// shouldn't get here - there will have been errors
-		return
+		return null
 	}
 
 	// Now that all targets from original changelist description are computed, compare to requested target branch, if applicable
@@ -700,7 +722,7 @@ export function computeTargets(
 	}
 
 	// compute final merge modes of all targets
-	const mergeActions: MergeAction[] = []
+	let mergeActions: MergeAction[] = []
 	const allDownstream = new Set<Branch>()
 
 	for (const [direct, indirectTargets] of merges.entries()) {
@@ -719,6 +741,13 @@ export function computeTargets(
 		mergeActions.push({branch: direct, mergeMode, furtherMerges, flags})
 	}
 
+	// If forceStompChanges, add the target as a stomp if it isn't a flagged target
+	if (info.forceStompChanges) {
+		mergeActions = mergeActions
+			.filter(action => action.mergeMode === 'normal')
+			.map(action => ({...action, mergeMode: 'clobber'}))
+	}
+
 	const branchGraph = sourceBranch.parent
 
 	// add indirect forced branches to allDownstream
@@ -731,13 +760,134 @@ export function computeTargets(
 		info.errors = errors
 	}
 
-	// report targets
-	if (mergeActions.length > 0) {
-		info.targets = mergeActions
-	}
-
 	// provide info on all targets that should eventually be merged too
 	if (allDownstream && allDownstream.size !== 0) {
 		info.allDownstream = [...allDownstream]
 	}
+
+	return [mergeActions, flags]
+}
+
+const colors = require('colors')
+colors.enable()
+colors.setTheme(require('colors/themes/generic-logging.js'))
+
+export function runTests(parentLogger: ContextualLogger) {
+	const logger = parentLogger.createChild('Targets')
+
+	const defaultArgs = {
+		isDefaultBot: true,
+		graphBotName: 'test',
+		cl: 1,
+		aliasUpper: '',
+		macros: {'m': ['#robomerge X, Y'], 'otherbot': ['#robomerge[A] B']},
+		logger
+	}
+
+	let fails = 0
+	let assertions = 0
+	let testName = ''
+	const assert = (b: boolean, msg: string, details: string[] = []) => {
+		if (!b) {
+			logger.error(`"${testName}" failed: ${colors.error(msg)}` + 
+				(details || []).map(s => `\t${s}\n`).join(''))
+			++fails
+		}
+		++assertions
+	}
+
+	const parserAssert = (parser: DescriptionParser,
+							descLines: number, commands: number, otherBot: number, macros: number, errors: number,
+							additional: boolean) => {
+		assert(
+			parser.descFinal.length === descLines &&
+			parser.arguments.length === commands &&
+			parser.otherBotArguments.length === otherBot &&
+			parser.expandedMacros.length === macros &&
+			parser.errors.length === errors &&
+			additional, 'Counts mismatch:', [
+				`descLines: got ${parser.descFinal.length}, expected ${descLines}`,
+				`commands: got ${parser.arguments.length}, expected ${commands}`,
+				`otherBot: got ${parser.otherBotArguments.length}, expected ${otherBot}`,
+				`macros: got ${parser.expandedMacros.length}, expected ${macros}`,
+				`errors: got ${parser.errors.length}, expected ${errors}`
+			]
+		)
+	}
+
+	(() => {
+		testName = 'no commands'
+		const parser = parseDescriptionLines({lines: ['no command'], ...defaultArgs})
+		parserAssert(parser, 1, 0, 0, 0, 0, parser.descFinal[0] === 'no command')
+	})();
+
+	(() => {
+		testName = 'multiple targets'
+		const parser = parseDescriptionLines({lines: ['#robomerge A', '#robomerge B'], ...defaultArgs})
+		parserAssert(parser, 0, 2, 0, 0, 0, true)
+	})();
+
+	(() => {
+		testName = 'comma separated targets'
+		const parser = parseDescriptionLines({lines: ['#robomerge A, B'], ...defaultArgs})
+		parserAssert(parser, 0, 2, 0, 0, 0, true)
+	})();
+
+	(() => {
+		testName = 'space separated targets'
+		const parser = parseDescriptionLines({lines: ['#robomerge A B'], ...defaultArgs})
+		parserAssert(parser, 0, 2, 0, 0, 0, true)
+	})();
+
+	(() => {
+		testName = 'macro'
+		const parser = parseDescriptionLines({lines: ['#robomerge M'], ...defaultArgs})
+		parserAssert(parser, 0, 2, 0, 1, 0, true)
+	})();
+
+	(() => {
+		testName = 'other bot'
+		const parser = parseDescriptionLines({lines: ['#robomerge[B] A'], ...defaultArgs})
+		parserAssert(parser, 0, 0, 1, 0, 0, true)
+	})();
+
+	(() => {
+		testName = 'all bots'
+		const parser = parseDescriptionLines({lines: ['#robomerge[ALL] A'], ...defaultArgs})
+		parserAssert(parser, 0, 1, 0, 0, 0, true)
+	})();
+
+	(() => {
+		testName = 'other bot macro'
+		const parser = parseDescriptionLines({lines: ['#robomerge otherbot'], ...defaultArgs})
+		parserAssert(parser, 0, 0, 1, 1, 0, true)
+	})();
+
+	(() => {
+		testName = 'unknown tag'
+		const parser = parseDescriptionLines({lines: ['#unknown blah'], ...defaultArgs})
+		parserAssert(parser, 1, 0, 0, 0, 0, parser.descFinal[0] === '#unknown blah')
+	})();
+
+	(() => {
+		testName = 'sanitize known tag'
+		const parser = parseDescriptionLines({lines: ['#fyi blah'], ...defaultArgs})
+		assert(parser.descFinal.length === 1 &&
+			parser.descFinal[0].indexOf('#') < 0 && parser.descFinal[0].toLowerCase().indexOf('fyi') >= 0,
+			'known tag sanitized')
+	})();
+
+	(() => {
+		testName = 'sanitize swarm tag'
+		const parser = parseDescriptionLines({lines: ['#review-123 blah'], ...defaultArgs})
+		assert(parser.descFinal.length === 1 &&
+			parser.descFinal[0].indexOf('#') < 0 && parser.descFinal[0].toLowerCase().indexOf('blah') >= 0,
+			'swarm tag sanitized')
+	})();
+
+	if (fails === 0) {
+		logger.info(colors.info(`Targets tests succeeded (${assertions} assertions)`))
+	}
+
+	return fails
 }
\ No newline at end of file
diff --git a/Engine/Extras/RoboMerge/v3/src/robo/watchdog.ts b/Engine/Extras/RoboMerge/v3/src/robo/watchdog.ts
--- a/Engine/Extras/RoboMerge/v3/src/robo/watchdog.ts
+++ b/Engine/Extras/RoboMerge/v3/src/robo/watchdog.ts
@@ -111,6 +111,7 @@ class Watchdog {
 
 	cbMap = new Map<number, Function>()
 	nextCbId = 1
+	lastError = ""
 
 	private readonly watchdogLogger = new ContextualLogger('Watchdog')
 
@@ -250,7 +251,7 @@ class Watchdog {
 		this.child.once('close', (code, signal) => {
 			this.childExited()
 			if (code !== null) {
-				this.watchdogLogger.warn(`Child process exited with code=${code}`)
+				this.watchdogLogger.warn(`Child process exited with code=${code}\n${this.lastError}`);
 				if (code !== 0) {
 					this.respawn()
 				}
@@ -299,6 +300,7 @@ class Watchdog {
 			}
 		})
 		this.child.stderr!.on('data', (chunk) => {
+			this.lastError = chunk.toString();
 			process.stderr.write(chunk)
 		})
 		this.child.stdout!.on('data', (chunk) => {
diff --git a/Engine/Extras/RoboMerge/v3/web/src/branches.ts b/Engine/Extras/RoboMerge/v3/web/src/branches.ts
--- a/Engine/Extras/RoboMerge/v3/web/src/branches.ts
+++ b/Engine/Extras/RoboMerge/v3/web/src/branches.ts
@@ -169,7 +169,7 @@ export function updateBranchList(graphBotName: string | null = null) {
 				.addClass('btn btn-xs robomerge-crash-buttons')
 				.text("View Logs")
 				.click(function() {
-					window.location.href = "/api/last_crash";
+					window.location.href = "/api/logs";
 				})
 				.appendTo(bl)
 
diff --git a/Engine/Extras/RoboMerge/v3/web/src/flow.ts b/Engine/Extras/RoboMerge/v3/web/src/flow.ts
--- a/Engine/Extras/RoboMerge/v3/web/src/flow.ts
+++ b/Engine/Extras/RoboMerge/v3/web/src/flow.ts
@@ -2,18 +2,24 @@
 
 import { BranchDefForStatus, BranchStatus } from "../../src/robo/status-types"
 
-type Branch = {
-	id: string
-	forcedDests: Set<string>
-	defaultDests: Set<string>
-	blockAssetDests: Set<string>
-}
-
-type Info = {
-	id: string
-	tooltip: string
-	name: string
+class Info {
 	graphNodeColor?: string
+
+	forcedDests = new Set<string>()
+	defaultDests = new Set<string>()
+	blockAssetDests = new Set<string>()
+
+	importance = 0
+
+
+	constructor(public id: string, public tooltip: string, public name: string, public sources: BranchStatus[]) {
+	}
+
+	update(branch: BranchStatus, getId: (alias: string) => string) {
+		this.forcedDests = new Set([...this.forcedDests, ...branch.def.forceFlowTo.map(getId)])
+		this.defaultDests = new Set([...this.defaultDests, ...branch.def.defaultFlow.map(getId)])
+		this.blockAssetDests = new Set([...this.blockAssetDests, ...branch.def.blockAssetTargets.map(getId)])
+	}
 }
 
 type NodeInfo = {
@@ -54,68 +60,81 @@ function renderGraph(src: string): Promise<any> {
 function decoratedAlias(bot: string, alias: string) {
 	return bot + ':' + alias
 }
-/** */
+
+function aliasForBranch(branch: BranchStatus) {
+	return decoratedAlias(branch.bot, branch.def.upperName)
+}
+
+
+type Options = {
+	hideDisconnected?: boolean
+	noGroups?: boolean
+	showOnlyForced?: boolean
+	botsToShow?: string[]
+}
+
+function parseOptions(search: string): Options {
+	return {
+		showOnlyForced: !!search.match(/showOnlyForced/i),
+		noGroups: !!search.match(/noGroups/i),
+		hideDisconnected: !!search.match(/hideDisconnected/i)
+	}
+}
+
 class Graph {
 	readonly nodeLabels = new Map<string, Info[]>()
 	readonly links: Link[] = []
 
 	private readonly aliases = new Map<string, Info>()
-	private readonly branchesById = new Map<string, Branch>()
+	private readonly allInfos: Info[] = []
 
 	private readonly linkOutward = new Set<string>()
 	private readonly linkInward = new Set<string>()
 
-	private readonly nodeImportance = new Map<string, number>()
+	private readonly connectedNodes = new Set<string>()
 
 	private branchList: BranchStatus[] = []
 
 	constructor(
 		private allBranches: BranchStatus[],
-		private whitelist: string[] | null = null
+		private options: Options
 	) {
 
 	}
 
 	singleBot(botName: string) {
 		this.branchList = this.allBranches.filter(b => b.bot === botName)
-
-		// special case branches named Main
-		this.nodeImportance.set(this.getIdForBot(botName)('MAIN'), 10)
-
-		for (const branchStatus of this.branchList) {
-			this.addBranch(branchStatus)
-		}
-		this.createIds()
-
-		for (const branchStatus of this.branchList) {
-			const getId = this.getIdForBot(branchStatus.bot)
-			for (const flow of branchStatus.def.flowsTo) {
-				this.addEdge(branchStatus, getId(flow))
-			}
-		}
-
+		this.addNodesAndEdges(false)
 		return this.makeGraph(botName, botName + 'integration paths')
 	}
 
-	allBots(whitelist: string[] | null = null) {
-		this.branchList = whitelist
-			? this.allBranches.filter(b => whitelist.indexOf(b.bot) >= 0)
+	allBots() {
+		this.branchList = this.options.botsToShow
+			? this.allBranches.filter(b => this.options.botsToShow!.indexOf(b.bot) >= 0)
 			: this.allBranches
 
 		this.findSharedNodes()
+		this.addNodesAndEdges(!this.options.noGroups)
+
+		return this.makeGraph('All bots', 'Flow including shared bots')
+	}
+
+	private addNodesAndEdges(grouped: boolean) {
 		for (const branchStatus of this.branchList) {
-			this.addBranch(branchStatus, branchStatus.bot)
+			this.addBranch(branchStatus, grouped ? branchStatus.bot : null)
 		}
-		this.createIds()
+
+		for (const info of this.allInfos) {
+			for (const sourceBranch of info.sources) {
+				info.update(sourceBranch, this.getIdForBot(sourceBranch.bot))
+			}
+		} 
 
 		for (const branchStatus of this.branchList) {
-			const getId = this.getIdForBot(branchStatus.bot)
 			for (const flow of branchStatus.def.flowsTo) {
-				this.addEdge(branchStatus, getId(flow))
+				this.addEdge(branchStatus, flow)
 			}
 		}
-
-		return this.makeGraph('All bots', 'Flow including shared bots')
 	}
 
 	private makeGraph(title: string, tooltip: string) {
@@ -134,8 +153,8 @@ class Graph {
 			nodeGroups.push([groupName, nodeInfo])
 
 			for (const info of v) {
-				const {id, tooltip, name} = info
-				let factor = (Math.min(this.nodeImportance.get(id) || 0, 10) - 1) / 9
+				const {id, tooltip, name, importance} = info
+				let factor = (Math.min(importance, 10) - 1) / 9
 				nodeInfo.push({
 					importance: factor,
 					marginX: (.2 * (1 - factor) + .4 * factor).toPrecision(1),
@@ -154,6 +173,9 @@ class Graph {
 			}
 
 			for (const info of nodeInfo) {
+				if (this.options.hideDisconnected && !this.connectedNodes.has(info.info.id)) {
+					continue
+				}
 				const attrs: [string, string][] = [
 					['label', `"${info.info.name}"`],
 					['tooltip', `"${info.info.tooltip}"`],
@@ -181,7 +203,7 @@ class Graph {
 		for (const link of this.links) {
 			// when there's forced/unforced pair, only the former is a constrait
 			// (this makes flow )
-			const combo = link.dst + link.src
+			const combo = link.dst + '->' + link.src
 
 			if (combo && this.linkOutward.has(combo) && this.linkInward.has(combo)) {
 				link.styles.push(['constraint', 'false'])
@@ -204,14 +226,21 @@ class Graph {
 		if (branch.aliases && branch.aliases.length !== 0) {
 			tooltip += ` (${branch.aliases.join(', ')})`
 		}
-		const info: Info = {
-			id: `_${branchStatus.bot}_${branch.upperName.replace(/[^\w]+/g, '_')}`,
-			tooltip, name: branch.name
-		}
+		const info = new Info(
+			`_${branchStatus.bot}_${branch.upperName.replace(/[^\w]+/g, '_')}`,
+			tooltip, branch.name,
+			[branchStatus]
+		)
+		this.allInfos.push(info)
 		if (branch.config.graphNodeColor) {
 			info.graphNodeColor = branch.config.graphNodeColor
 		}
 
+		if (branch.upperName === 'MAIN') {
+			// special case branches named Main
+			info.importance = 10
+		}
+
 		// note: branch.upperName is always in branch.aliases
 		for (const alias of branch.aliases) {
 			this.aliases.set(decoratedAlias(branchStatus.bot, alias), info)
@@ -221,28 +250,37 @@ class Graph {
 
 	// dstName 
 	private addEdge(srcBranchStatus: BranchStatus, dst: string) {
-		const src = this.getIdForBranch(srcBranchStatus)
-		const branch = this.branchesById.get(src)!
+		const src = aliasForBranch(srcBranchStatus)
+		const srcInfo = this.aliases.get(src)!
+
+		const dstInfo = this.aliases.get(decoratedAlias(srcBranchStatus.bot, dst))!
+
+		++srcInfo.importance
+		++dstInfo.importance
+
+		const isForced = srcInfo.forcedDests.has(dstInfo.id)
+		if (!isForced && this.options.showOnlyForced) {
+			return
+		}
 
-		this.nodeImportance.set(src, (this.nodeImportance.get(src) || 0) + 1)
-		this.nodeImportance.set(dst, (this.nodeImportance.get(dst) || 0) + 1)
+		this.connectedNodes.add(srcInfo.id)
+		this.connectedNodes.add(dstInfo.id)
 
-		const isForced = branch.forcedDests.has(dst)
 		const edgeStyle =
 			srcBranchStatus.def.convertIntegratesToEdits ?	'roboshelf' :
 			isForced ?										'forced' :
-			branch.defaultDests.has(dst) ?					'defaultFlow' :
-			branch.blockAssetDests.has(dst) ?				'blockAssets' : 'onRequest'
+			srcInfo.defaultDests.has(dstInfo.id) ?			'defaultFlow' :
+			srcInfo.blockAssetDests.has(dstInfo.id) ?		'blockAssets' : 'onRequest'
 
 		const styles = [...EDGE_STYLES[edgeStyle]]
 
-		const link: Link = {src, dst, styles}
+		const link: Link = {src: srcInfo.id, dst: dstInfo.id, styles}
 		this.links.push(link)
 		if (isForced) {
-			this.linkOutward.add(src + dst)
+			this.linkOutward.add(link.src + '->' + link.dst)
 		}
 		else if (edgeStyle === 'blockAssets' || edgeStyle === 'onRequest') {
-			this.linkInward.add(dst + src)
+			this.linkInward.add(link.dst + '->' + link.src)
 		}
 	}
 
@@ -261,15 +299,13 @@ class Graph {
 		for (const [k, v] of streamMap) {
 			const [color, streams] = v
 			if (streams.length > 1) {
-				const info: Info = {
-					id: k.replace(/[^\w]+/g, '_'), tooltip: 'tooltip todo', name: k
-				}
+				const info = new Info(k.replace(/[^\w]+/g, '_'), streams.map(s => s.bot).join(', '), k, streams)
+				this.allInfos.push(info)
 				if (color) {
 					info.graphNodeColor = color
 				}
 
 				for (const branchStatus of streams) {
-
 					for (const alias of branchStatus.def.aliases) {
 						this.aliases.set(decoratedAlias(branchStatus.bot, alias), info)
 					}
@@ -280,25 +316,6 @@ class Graph {
 		}
 	}
 
-	private createIds() {
-		for (const branchStatus of this.branchList) {
-			const branch = branchStatus.def
-			const getId = this.getIdForBot(branchStatus.bot)
-			const branchId = getId(branch.upperName)
-			this.branchesById.set(branchId, {
-				id: branchId,
-				forcedDests: new Set<string>(branch.forceFlowTo.map(getId)),
-				defaultDests: new Set<string>(branch.defaultFlow.map(getId)),
-				blockAssetDests: new Set<string>(branch.blockAssetTargets ? branch.blockAssetTargets.map(getId) : [])
-			})
-		}
-	}
-
-	private getIdForBranch(branch: BranchStatus) {
-		// note: branch.upperName is always in branch.aliases
-		return this.getIdForBot(branch.bot)(branch.def.upperName)
-	}
-
 	private getIdForBot(bot: string) {
 		return (alias: string) => {
 			const info = this.aliases.get(decoratedAlias(bot, alias))
@@ -307,24 +324,28 @@ class Graph {
 	}
 }
 
-function makeGraph(data: BranchStatus[], botNameParam: string | string[]) {
+type Args = ({ botNames: string[] } | { singleBotName: string }) & Options
 
-	let whitelist: string[] | null = null
-	let botName: string | null
+function makeGraph(data: BranchStatus[], args: Args) {
 
-	const graph = new Graph(data)
-	if (Array.isArray(botNameParam)) {
-		botName = null
-		if (botNameParam.length > 0) {
-			return graph.allBots(botNameParam.map(s => s.toUpperCase()))
-		}
-		return graph.allBots()
+	let botsToShow: string[] | null = null
+
+	const options: Options = args
+
+	const singleBotName = (args as { singleBotName: string }).singleBotName
+	if (singleBotName) {
+		return (new Graph(data, options)).singleBot(singleBotName)
+	}
+
+	const botNames = (args as { botNames: string[] }).botNames
+	if (botNames.length > 0) {
+		options.botsToShow = botNames.map(s => s.toUpperCase())
 	}
-	return graph.singleBot(botNameParam.toUpperCase())
+	return (new Graph(data, options)).allBots()
 }
 
-export function showFlowGraph(data: BranchStatus[], botNameParam: string | string[]) {
-	const lines = makeGraph(data, botNameParam)
+function showFlowGraph(data: BranchStatus[], args: Args) {
+	const lines = makeGraph(data, args)
 
 	const graphContainer = $('<div class="clearfix">')
 	const flowGraph = $('<div class="flow-graph" style="display:inline-block;">').appendTo(graphContainer)
