[NO NON-PII COMMIT MESSAGE]

diff --git a/Engine/Extras/RoboMerge/v3/src/robo/targets.ts b/Engine/Extras/RoboMerge/v3/src/robo/targets.ts
--- a/Engine/Extras/RoboMerge/v3/src/robo/targets.ts
+++ b/Engine/Extras/RoboMerge/v3/src/robo/targets.ts
@@ -241,7 +241,7 @@ export class DescriptionParser {
 		if (token.tag.match(/REVIEW-(\d+)/)) {
 							// remove swarm review numbers entirely (they can't be neutered and generate emails)
 			if (token.rest) {
-				this.descFinal.push(token.rest)
+				this.descFinal.push(`${token.initialWhitespace}[REVIEW] ${token.rest}`)
 			}
 			return false
 		}
diff --git a/Engine/Extras/RoboMerge/v3/src/robo/unit-tests.ts b/Engine/Extras/RoboMerge/v3/src/robo/unit-tests.ts
--- a/Engine/Extras/RoboMerge/v3/src/robo/unit-tests.ts
+++ b/Engine/Extras/RoboMerge/v3/src/robo/unit-tests.ts
@@ -4,6 +4,7 @@ import { ContextualLogger } from '../common/logger';
 import { runTests as gate_runTests } from './gate';
 import { runTests as graph_runTests } from '../new/graph';
 import { runTests as notifications_runTests } from './notifications';
+import { runTests as targets_runTests } from './targets';
 import { runTests as ztag_runTests } from '../common/ztag';
 
 const unitTestsLogger = new ContextualLogger('Unit Tests')
@@ -12,6 +13,7 @@ const tests = [
 	gate_runTests,
 	graph_runTests,
 	notifications_runTests,
+	targets_runTests,
 	ztag_runTests
 ]
 
