[NO NON-PII COMMIT MESSAGE]

diff --git a/Engine/Extras/RoboMerge/v3/src/robo/session.ts b/Engine/Extras/RoboMerge/v3/src/robo/session.ts
--- a/Engine/Extras/RoboMerge/v3/src/robo/session.ts
+++ b/Engine/Extras/RoboMerge/v3/src/robo/session.ts
@@ -92,8 +92,14 @@ export class Session {
 			Object.assign(config, Session.LDAP_CONFIG)
 			const auth = new LdapAuth(config)
 
+			logger.info(`checking LDAP groups for user "${creds.user}"`)
+
+			const startTime = Date.now()
 			auth.on('error', fail)
 			auth.authenticate(creds.user, creds.password, (err: any | null, userData: any) => {
+				const duration = Math.round((Date.now() - startTime) / 1000)
+				logger.info(`LDAP for "${creds.user}" took ${duration}s`)
+
 				if (err) {
 					if (err.name === 'InvalidCredentialsError' ||
 						(typeof(err) === 'string' && err.startsWith('no such user'))) {
