[NO NON-PII COMMIT MESSAGE]

diff --git a/Engine/Extras/RoboMerge/v3/src/robo/session.ts b/Engine/Extras/RoboMerge/v3/src/robo/session.ts
--- a/Engine/Extras/RoboMerge/v3/src/robo/session.ts
+++ b/Engine/Extras/RoboMerge/v3/src/robo/session.ts
@@ -88,11 +88,14 @@ export class Session {
 	static login(creds: Credentials, logger: ContextualLogger) {
 		return new Promise<string | null>((done, fail) => {
 			// LdapAuth modifies the config, so copy to be clean (probably doesn't matter)
+			logger.info(`checking LDAP groups for user "${creds.user}"`)
+
 			const config: any = {}
 			Object.assign(config, Session.LDAP_CONFIG)
 			const auth = new LdapAuth(config)
 
-			logger.info(`checking LDAP groups for user "${creds.user}"`)
+			logger.info(`LDAP: auth helper created (${creds.user} log-in)`)
+
 
 			const startTime = Date.now()
 			auth.on('error', fail)
