[NO NON-PII COMMIT MESSAGE]

diff --git a/Engine/Extras/RoboMerge/v3/src/common/perforce.ts b/Engine/Extras/RoboMerge/v3/src/common/perforce.ts
--- a/Engine/Extras/RoboMerge/v3/src/common/perforce.ts
+++ b/Engine/Extras/RoboMerge/v3/src/common/perforce.ts
@@ -549,7 +549,8 @@ export class PerforceContext {
 	// find a workspace for the given user
 	// output format is an array of workspace names
 	async find_workspaces(user?: string) {
-		let parsedClients = await this._execP4Ztag(null, ['clients', '-u', user || this.username], { multiline: true });
+		// -a to include workspaces on edge servers
+		let parsedClients = await this._execP4Ztag(null, ['clients', '-u', user || this.username, '-a'], { multiline: true });
 		let workspaces = [];
 		for (let clientDef of parsedClients) {
 			if (clientDef.client) {
diff --git a/Engine/Extras/RoboMerge/v3/src/robo/robo.ts b/Engine/Extras/RoboMerge/v3/src/robo/robo.ts
--- a/Engine/Extras/RoboMerge/v3/src/robo/robo.ts
+++ b/Engine/Extras/RoboMerge/v3/src/robo/robo.ts
@@ -36,7 +36,7 @@ VersionReader.init(roboStartupLogger)
 
 
 const USER_WORKSPACE_EXCLUDE_PATTERNS: (RegExp | string)[] = [
-	'horde-p4bridge-hordeserver-'
+	'horde-p4bridge-'
 ]
 
 // I seem to have broken this
