[NO NON-PII COMMIT MESSAGE]

diff --git a/Engine/Extras/RoboMerge/v3/src/robo/branch-interfaces.ts b/Engine/Extras/RoboMerge/v3/src/robo/branch-interfaces.ts
--- a/Engine/Extras/RoboMerge/v3/src/robo/branch-interfaces.ts
+++ b/Engine/Extras/RoboMerge/v3/src/robo/branch-interfaces.ts
@@ -74,7 +74,7 @@ export interface BranchGraphInterface {
 	getBranchNames(): string[]
 }
 
-export interface EditableBranch extends BranchBase {
+export type EditableBranch = BranchBase & {
 	bot?: NodeBotInterface
 	parent: BranchGraphInterface
 	workspace: RoboWorkspace
diff --git a/Engine/Extras/RoboMerge/v3/src/robo/branchdefs.ts b/Engine/Extras/RoboMerge/v3/src/robo/branchdefs.ts
--- a/Engine/Extras/RoboMerge/v3/src/robo/branchdefs.ts
+++ b/Engine/Extras/RoboMerge/v3/src/robo/branchdefs.ts
@@ -2,6 +2,12 @@
 
 import { StreamSpecs } from '../common/perforce'
 
+// for validating all branchspecs
+// import { ContextualLogger } from '../common/logger';
+// import { initializePerforce, PerforceContext, StreamSpec } from '../common/perforce';
+// import fs = require('fs')
+// import path = require('path')
+
 const jsonlint: any = require('jsonlint')
 
 const RESERVED_BRANCH_NAMES = ['NONE', 'DEFAULT', 'IGNORE', 'DEADEND', ''];
@@ -26,22 +32,24 @@ export interface BotConfig {
 	macros: { [name: string]: string[] }
 }
 
-export interface BranchBase {
-	name: string
+const branchBasePrototype = {
+	name: '',
 
-	rootPath: string
-	isDefaultBot: boolean
-	emailOnBlockage: boolean // if present, completely overrides BotConfig
-
-	notify: string[]
-	flowsTo: string[]
-	forceFlowTo: string[]
-	defaultFlow: string[]
-	resolver: string | null
-	aliases: string[]
-	badgeProject: string | null
+	rootPath: '',
+	isDefaultBot: false,
+	emailOnBlockage: false, // if present, completely overrides BotConfig
+
+	notify: [''],
+	flowsTo: [''],
+	forceFlowTo: [''],
+	defaultFlow: [''],
+	resolver: '' as string | null,
+	aliases: [''],
+	badgeProject: '' as string | null,
 }
 
+export type BranchBase = typeof branchBasePrototype
+
 export class IntegrationMethod {
 	static NORMAL = 'normal'
 	static CONVERT_TO_EDIT = 'convert-to-edit'
@@ -65,56 +73,75 @@ export type IntegrationWindowPane = {
 	durationHours: number
 }
 
-export type CommonOptionFields = {
-	lastGoodCLPath: string | number
-	pauseCISUnlessAtGate: boolean
 
-	initialCL: number
-	forcePause: boolean
+export const commonOptionFieldsPrototype = {
+	lastGoodCLPath: 0 as string | number,
+	pauseCISUnlessAtGate: false,
+
+	initialCL: 0,
+	forcePause: false,
 
-	disallowSkip: boolean
-	incognitoMode: boolean
+	disallowSkip: false,
+	incognitoMode: false,
 
-	excludeAuthors: string[] // if present, completely overrides BotConfig
+	excludeAuthors: [] as string[], // if present, completely overrides BotConfig
 
 	// by default, specify when gate catch ups are allowed; can be inverted to disallow
-	integrationWindow: IntegrationWindowPane[]
-	invertIntegrationWindow: boolean
+	integrationWindow: [] as IntegrationWindowPane[],
+	invertIntegrationWindow: false,
+
+	// fake property
+	_comment: ''
 }
 
-type NodeOptionFields = BranchBase & CommonOptionFields & {
-	disabled: boolean
-	integrationMethod: string
-	forceAll: boolean
-	visibility: string[] | string
-	blockAssetFlow: string[]
-	disallowDeadend: boolean
+export type CommonOptionFields = typeof commonOptionFieldsPrototype
+
+const nodeOptionFieldsPrototype = {
+	...branchBasePrototype,
+	...commonOptionFieldsPrototype,
 
-	streamDepot: string
-	streamName: string
-	streamSubpath: string
-	workspace: (string | null)
+	disabled: false,
+	integrationMethod: '',
+	forceAll: false,
+	visibility: '' as string[] | string,
+	blockAssetFlow: [''],
+	disallowDeadend: false,
 
-	graphNodeColor: string
+	streamDepot: '',
+	streamName: '',
+	streamSubpath: '',
+	workspace: '',
+
+	graphNodeColor: '',
 
 	// if set, still generate workspace but use this name
-	workspaceNameOverride: string
-	additionalSlackChannelForBlockages: string
-	ignoreBranchspecs: boolean
+	workspaceNameOverride: '',
+	additionalSlackChannelForBlockages: '',
+	ignoreBranchspecs: false,
 
-	badgeUrlOverride: string
+	badgeUrlOverride: '',
 }
 
+const nodeOptionFieldNames: ReadonlySet<string> = new Set(Object.keys(nodeOptionFieldsPrototype));
+
+type NodeOptionFields = typeof nodeOptionFieldsPrototype
+
 // will eventually have all properties listed on wiki
-type EdgeOptionFields = CommonOptionFields & {
-	additionalSlackChannel: string
+const edgeOptionFieldsPrototype = {
+	...commonOptionFieldsPrototype,
+
+	additionalSlackChannel: '',
 
-	terminal: boolean // changes go along terminal edges but no further
-	doHackyOkForGithubThing: boolean
+	terminal: false, // changes go along terminal edges but no further
+	doHackyOkForGithubThing: false,
 
-	implicitCommands: string[]
+	implicitCommands: ['']
 }
 
+const edgeOptionFieldNames: ReadonlySet<string> = new Set(Object.keys(edgeOptionFieldsPrototype));
+
+type EdgeOptionFields = typeof edgeOptionFieldsPrototype
+
 export type NodeOptions = Partial<NodeOptionFields>
 export type EdgeOptions = Partial<EdgeOptionFields>
 
@@ -361,6 +388,13 @@ export class BranchDefs {
 				BranchDefs.checkValidIntegrationMethod(outErrors, def.integrationMethod!, def.name!)
 			}
 
+			// check all properties are known (could make things case insensitive here)
+			for (const keyName of Object.keys(def)) {
+				if (!nodeOptionFieldNames.has(keyName)) {
+					throw new Error(`Unknown property '${keyName}' specified for node ${def.name}`)
+				}
+			}
+
 			validateCommonOptions(def)
 		}
 
@@ -374,6 +408,18 @@ export class BranchDefs {
 					outErrors.push('Unrecognised target node in edge property ' + edge.to)
 				}
 
+				// check all properties are known (could make things case insensitive here)
+				for (const keyName of Object.keys(edge)) {
+					if (keyName !== 'from' && keyName !== 'to' 
+
+// temporarily allow resolver until supported properly
+&& keyName !== 'resolver'
+
+						&& !edgeOptionFieldNames.has(keyName)) {
+						throw new Error(`Unknown property '${keyName}' specified for edge ${edge.from}->${edge.to}`)
+					}
+				}
+
 				validateCommonOptions(edge)
 			}
 		}
@@ -391,7 +437,7 @@ export class BranchDefs {
 						flowsTo.add(branchName)
 					}
 					else {
-						outErrors.push(`'${def.name}' flows to unknown branch/alias '${to}'`)				
+						outErrors.push(`'${def.name}' flows to unknown branch/alias '${to}'`)
 					}
 				}
 			}
@@ -403,10 +449,10 @@ export class BranchDefs {
 				else for (const to of def.forceFlowTo) {
 					const branchName = names.get(to.toUpperCase())
 					if (!branchName) {
-						outErrors.push(`'${def.name}' force flows to unknown branch/alias '${to}'`)				
+						outErrors.push(`'${def.name}' force flows to unknown branch/alias '${to}'`)
 					}
 					else if (!flowsTo.has(branchName)) {
-						outErrors.push(`'${def.name}' force flows but does not flow to '${to}'`)				
+						outErrors.push(`'${def.name}' force flows but does not flow to '${to}'`)
 					}
 				}
 			}
@@ -415,6 +461,14 @@ export class BranchDefs {
 		// Check branchspecs for valid branches
 		if (branchGraph.branchspecs) {
 			for (const spec of branchGraph.branchspecs) {
+				for (const [key, val] of Object.entries(spec)) {
+					if (key !== 'from' && key !== 'to' && key !== 'name') {
+						outErrors.push('Unexpected branchspec property: ' + key)
+					}
+					if (typeof val !== 'string') {
+						outErrors.push(`Branchspec property ${key} is not a string`)
+					}
+				}
 				if (!spec.from || !spec.to) {
 					outErrors.push(`Invalid branchspec ${spec.name} (requires both to and from fields)`)
 				}
@@ -437,3 +491,34 @@ export class BranchDefs {
 		return {branchGraphDef: branchGraph, config: defaultConfigForWholeBot}
 	}
 }
+
+// function verifyAllBranchmaps(allStreamSpecs: Map<string, StreamSpec>, folder: string) {
+
+// 	for (const file of fs.readdirSync(folder, {encoding: "utf8"})) {
+// 		if (!file.endsWith('branchmap.json') || file.indexOf('iron.') >= 0) {
+// 			continue
+// 		}
+// 		console.log(file)
+// 		const validationErrors: string[] = []
+
+// 		const result = BranchDefs.parseAndValidate(
+// 			validationErrors,
+// 			fs.readFileSync(path.join(folder, file), 'utf8'),
+// 			allStreamSpecs)
+
+// 		if (!result.branchGraphDef) {
+// 			throw new Error(validationErrors.length === 0 ? 'Failed to parse' : validationErrors.join('\n'))
+// 		}
+// 		console.log(`Branches found in ${file}: ${result.branchGraphDef.branches.length}`)
+// 	}
+// }
+
+// export async function runTests(logger: ContextualLogger) {
+// 	if (process.platform === 'darwin') {
+// 		await initializePerforce(logger)
+// 		const allStreamSpecs = await (new PerforceContext(logger)).streams()
+
+// 		verifyAllBranchmaps(allStreamSpecs, '../RoboMerge/data')
+// 	}
+// 	return 0
+// }
