[NO NON-PII COMMIT MESSAGE]

diff --git a/Engine/Extras/RoboMerge/v3/src/new/graph.ts b/Engine/Extras/RoboMerge/v3/src/new/graph.ts
--- a/Engine/Extras/RoboMerge/v3/src/new/graph.ts
+++ b/Engine/Extras/RoboMerge/v3/src/new/graph.ts
@@ -588,7 +588,7 @@ export function runTests(parentLogger: ContextualLogger) {
 	return fail
 }
 
-type ChangeFlag = 'review' | 'manual' | 'null' | 'ignore' | 'disregardexcludedauthors' | 'roboshelf'
+type ChangeFlag = 'manual' | 'null' | 'ignore' | 'disregardexcludedauthors' | 'roboshelf'
 const ALLOWED_RAW_FLAGS = new Set(['null','ignore','deadend'])
 
 const FLAGMAP: {[name: string]: ChangeFlag} = {
diff --git a/Engine/Extras/RoboMerge/v3/src/robo/edgebot.ts b/Engine/Extras/RoboMerge/v3/src/robo/edgebot.ts
--- a/Engine/Extras/RoboMerge/v3/src/robo/edgebot.ts
+++ b/Engine/Extras/RoboMerge/v3/src/robo/edgebot.ts
@@ -423,15 +423,11 @@ class EdgeBotImpl extends PerforceStatefulBot {
 			description += `#ROBOMERGE-CONFLICT from-shelf\n`
 		}
 
-
 		if (target.flags.has('disregardexcludedauthors')) {
 			description += '#ROBOMERGE[ALL]: #DISREGARDEXCLUDEDAUTHORS\n'
 		}
 
 		let flags = '' // not all flags propagate, build them piecemeal
-		if (target.flags.has('review')) {
-			flags += ' #REVIEW'
-		}
 
 		const thisBotMergeCommands = target.furtherMerges
 			.filter(target => !target.otherBot)
@@ -485,11 +481,6 @@ class EdgeBotImpl extends PerforceStatefulBot {
 		// if required, add author review here so they're not in target.description, which is used for shelf description in case of conflict
 		let desc = target.description! // target.description always ends in newline
 
-		if (target.flags.has('review')) {
-			const owner = getIntegrationOwner(this.targetBranch, info.owner) || info.author
-			desc += `#CodeReview: ${owner}\n`
-		}
-
 		// create a new CL
 		const changenum = await this.p4.new_cl(this.targetBranch.workspace, desc)
 
diff --git a/Engine/Extras/RoboMerge/v3/src/robo/nodebot.ts b/Engine/Extras/RoboMerge/v3/src/robo/nodebot.ts
--- a/Engine/Extras/RoboMerge/v3/src/robo/nodebot.ts
+++ b/Engine/Extras/RoboMerge/v3/src/robo/nodebot.ts
@@ -348,7 +348,7 @@ export class NodeBot extends PerforceStatefulBot implements NodeBotInterface {
 		if (fromQueue.timestamp) {
 			logMessage += ` (delay: ${Math.round((Date.now() - fromQueue.timestamp) / 1000)})`
 		}
-		
+	
 		let specifiedTargetBranch : Branch | null = null
 		if (fromQueue.targetBranchName) {
 			specifiedTargetBranch = this._getBranch(fromQueue.targetBranchName)
@@ -1414,18 +1414,6 @@ export class NodeBot extends PerforceStatefulBot implements NodeBotInterface {
 			}
 		}
 
-		if (info.targets) {
-			for (const target of info.targets) {
-				if (this.getImmediateEdge(target.branch)!.incognitoMode) {
-
-					if (target.flags.has('review')) {
-						info.errors = info.errors || []
-						info.errors.push(`Target ${target.branch.name} flagged for review - not allowed in incognito mode`)
-					}
-				}
-			}
-		}
-
 		// stop on errors and e-mail interested parties
 		if (info.errors) {
 			this.nodeBotLogger.debug(`Errors during createChangeInfo():\n${info.errors.join('\n')}`)
diff --git a/Engine/Extras/RoboMerge/v3/src/robo/targets.ts b/Engine/Extras/RoboMerge/v3/src/robo/targets.ts
--- a/Engine/Extras/RoboMerge/v3/src/robo/targets.ts
+++ b/Engine/Extras/RoboMerge/v3/src/robo/targets.ts
@@ -75,16 +75,11 @@ const ROBO_TAGS = [
 	'ROBOMERGE-SOURCE',
 ]
 
-type ChangeFlag = 'review' | 'manual' | 'null' | 'ignore' | 'disregardexcludedauthors'
+type ChangeFlag = 'manual' | 'null' | 'ignore' | 'disregardexcludedauthors'
 
 // mapping of #ROBOMERGE: flags to canonical names
 // use these with a pound like #ROBOMERGE: #stage
 const FLAGMAP: {[name: string]: ChangeFlag} = {
-	// force a codereview to be sent to the owner as commits flow down the automerge chain
-	review: 'review',
-	forcereview: 'review',
-	cr: 'review',
-
 	// don't merge, only shelf this change and review the owner
 	manual: 'manual',
 	nosubmit: 'manual',
