[NO NON-PII COMMIT MESSAGE]

diff --git a/Engine/Extras/RoboMerge/v3/README.md b/Engine/Extras/RoboMerge/v3/README.md
--- a/Engine/Extras/RoboMerge/v3/README.md
+++ b/Engine/Extras/RoboMerge/v3/README.md
@@ -58,7 +58,6 @@ using the 5.0 alias and allowing RoboMerge to calculate the route from Main to 5
 | -       | skip | branch | `#robomerge -7.40`        | Do not merge to the specified branch (used for branches that are merged to automatically)                                                                            |
 | ignore  |      | global | `#robomerge ignore`       | Completely ignore this changelist - no automatic merges happen. Note that a target branch of 'ignore' is interpreted as an ignore flag (can be used without a #).    |
 | deadend |      | global | `#robomerge deadend`      | Completely ignore this changelist - no automatic merges happen. Note that a target branch of 'deadend' is interpreted as a deadend flag (can be used without a #).   |
-| #review |      | global | `#robomerge #review 7.40` | Add #codereview for owner for each commit in merge chain                                                                                                             |
 | #manual |      | global | `#robomerge #manual`      | Do not commit merge - add #codereview and shelve for owner                                                                                                           |
 | !       | null | branch | `#robomerge !7.40`        | Perform a null merge to the specified branch, i.e. convince Perforce that a merge has happened, but don't actually make any changes to the content of the stream.    |
 | null    |      | global | `#robomerge null`         | Make every automatic merge of this commit a null merge (this is the old behavior of the deadend tag). Like ignore and deadend, can be used without a #.              |
