[NO NON-PII COMMIT MESSAGE]

diff --git a/Engine/Extras/RoboMerge/v3/src/common/perforce.ts b/Engine/Extras/RoboMerge/v3/src/common/perforce.ts
--- a/Engine/Extras/RoboMerge/v3/src/common/perforce.ts
+++ b/Engine/Extras/RoboMerge/v3/src/common/perforce.ts
@@ -1208,8 +1208,8 @@ export class PerforceContext {
 		await this._execP4(workspace, ['change', '-i', changeFlag], { stdin: form, quiet: true })
 	}
 
-	async where(roboWorkspace: RoboWorkspace, clientPath: string) {
-		return await this._execP4Ztag(roboWorkspace, ['where', clientPath])
+	where(roboWorkspace: RoboWorkspace, clientPath: string) {
+		return this._execP4Ztag(roboWorkspace, ['where', clientPath])
 	}
 
 	async filelog(roboWorkspace: RoboWorkspace, depotPath: string, beginRev: string, endRev: string, longOutput = false) {
