[NO NON-PII COMMIT MESSAGE]

diff --git a/Engine/Extras/RoboMerge/v3/functional_tests/framework/src/framework.ts b/Engine/Extras/RoboMerge/v3/functional_tests/framework/src/framework.ts
--- a/Engine/Extras/RoboMerge/v3/functional_tests/framework/src/framework.ts
+++ b/Engine/Extras/RoboMerge/v3/functional_tests/framework/src/framework.ts
@@ -193,6 +193,11 @@ type EdgeOptionFields = {
 	invertIntegrationWindow: boolean
 
 	implicitCommands: string[]
+
+	approval: {
+		description: string,
+		group: string
+	}
 }
 
 export type EdgeProperties = Partial<EdgeOptionFields> & {
@@ -648,6 +653,15 @@ export abstract class FunctionalTest {
 			.then((cls: number[]) => cls.indexOf(cl) >= 0)
 	}
 
+	async ensureConflictMessagePostedToSlack(sourceStream: string, targetStream: string) {
+		const edgeDisplayName = `${sourceStream} -> ${targetStream}`
+		this.info(`Ensuring conflict message sent for ${edgeDisplayName}`)
+		const edgeState = await this.getEdgeState('Release', 'Main')
+		if (!this.wasMessagePostedToSlack(this.botName.toLowerCase(), edgeState.conflictCl)) {
+			throw new Error(`no message sent for CL#${edgeState.conflictCl}`)
+		}
+	}
+
 	async verifyAndPerformStomp(source: string, target: string, additionalSlackChannel?: string) {
 		const edgeState: EdgeState = await this.getEdgeState(source, target)
 
diff --git a/Engine/Extras/RoboMerge/v3/functional_tests/framework/src/tests/approval.ts b/Engine/Extras/RoboMerge/v3/functional_tests/framework/src/tests/approval.ts
new file mode 100644
--- /dev/null
+++ b/Engine/Extras/RoboMerge/v3/functional_tests/framework/src/tests/approval.ts
@@ -0,0 +1,39 @@
+import { EdgeProperties, P4Util } from '../framework'
+import { SimpleMainAndReleaseTestBase, streams } from '../SimpleMainAndReleaseTestBase'
+
+export class Approval extends SimpleMainAndReleaseTestBase {
+
+	async setup() {
+		await this.p4.depot('stream', this.depotSpec())
+		await this.createStreamsAndWorkspaces(streams)
+		await P4Util.addFileAndSubmit(this.getClient('Main'), 'test.txt', 'Initial content')
+		await this.initialPopulate()
+	}
+
+	async run() {
+		const releaseClient = this.getClient('Release')
+		await releaseClient.sync()
+
+		// await this.p4.lock(this.getClient('Main').workspace, this.getStreamPath('Main'))
+		await P4Util.editFileAndSubmit(releaseClient, 'test.txt', 'Initial content\n\nmergeable')
+	}
+
+	verify() {
+		return Promise.all([
+				this.ensureNotBlocked('Main'),
+				this.ensureBlocked('Release', 'Main'),
+				this.ensureConflictMessagePostedToSlack('Release', 'Main')
+			])
+	}
+
+	getEdges() : EdgeProperties[] {
+		return [
+		  { from: this.fullBranchName('Release'), to: this.fullBranchName('Main')
+		  , approval: {
+		      description: 'hey there!',
+		      group: 'approvers'
+		    }
+		  }
+		]
+	}
+}
\ No newline at end of file
diff --git a/Engine/Extras/RoboMerge/v3/src/robo/branch-interfaces.ts b/Engine/Extras/RoboMerge/v3/src/robo/branch-interfaces.ts
--- a/Engine/Extras/RoboMerge/v3/src/robo/branch-interfaces.ts
+++ b/Engine/Extras/RoboMerge/v3/src/robo/branch-interfaces.ts
@@ -167,7 +167,10 @@ export interface Blockage {
 	failure: Failure
 	owner: string
 	ownerEmail: Promise<string | null>
-
+	approval?: {
+		group: string
+		shelfCl: number
+	}
 	time: Date
 }
 export type NodeOpUrlGenerator = (blockage: Blockage | null) => BlockageNodeOpUrls | null
diff --git a/Engine/Extras/RoboMerge/v3/src/robo/branchdefs.ts b/Engine/Extras/RoboMerge/v3/src/robo/branchdefs.ts
--- a/Engine/Extras/RoboMerge/v3/src/robo/branchdefs.ts
+++ b/Engine/Extras/RoboMerge/v3/src/robo/branchdefs.ts
@@ -135,7 +135,12 @@ const edgeOptionFieldsPrototype = {
 	terminal: false, // changes go along terminal edges but no further
 	doHackyOkForGithubThing: false,
 
-	implicitCommands: ['']
+	implicitCommands: [''],
+
+	approval: {
+		description: '',
+		group: ''
+	}
 }
 
 const edgeOptionFieldNames: ReadonlySet<string> = new Set(Object.keys(edgeOptionFieldsPrototype));
diff --git a/Engine/Extras/RoboMerge/v3/src/robo/edgebot.ts b/Engine/Extras/RoboMerge/v3/src/robo/edgebot.ts
--- a/Engine/Extras/RoboMerge/v3/src/robo/edgebot.ts
+++ b/Engine/Extras/RoboMerge/v3/src/robo/edgebot.ts
@@ -292,18 +292,6 @@ class EdgeBotImpl extends PerforceStatefulBot {
 		return results
 	}
 
-	private async handleIntegrationError(failure: Failure, pending: PendingChange) {
-		// Revert on Integration Error
-		if (pending.newCl > 0) {
-			await this.revertAndDelete(pending.newCl)
-			pending.newCl = -1
-		}
-		
-		// Pause Edge
-		const pauseInfo = this.createEdgeBlockageInfo(failure, pending)
-		this.block(pauseInfo, FAILED_CHANGELIST_PAUSE_TIMEOUT_SECONDS)
-	}
-
 	/** Handle any failures after a successful integration, e.g. conflicts or submit errors */
 	private async handlePostIntegrationFailure(failure: Failure, pending: PendingChange) {
 		const logMessage = `Post-integration failure while integrating CL ${pending.change.cl} to ${this.branch.name}`
@@ -540,8 +528,19 @@ class EdgeBotImpl extends PerforceStatefulBot {
 		else {
 			failure  = { kind: 'Integration error', description }
 		}
+
+		// Revert attempt
+		if (pending.newCl > 0) {
+			await this.revertAndDelete(pending.newCl)
+			pending.newCl = -1
+		}
 		
-		await this.handleIntegrationError(failure, pending)
+		if (!pending.change.userRequest) {
+			// Pause Edge
+			const pauseInfo = this.createEdgeBlockageInfo(failure, pending)
+			this.block(pauseInfo, FAILED_CHANGELIST_PAUSE_TIMEOUT_SECONDS)
+		}
+
 		// Send to source node to facilitate notification handling
 		if (this.sourceNode.handleMergeFailure(failure, pending)) {
 			this.gate.onBlockage()
@@ -612,18 +611,45 @@ class EdgeBotImpl extends PerforceStatefulBot {
 				}
 			}
 		}
-
-		if (!failure) {
+		else
+		{
+			// no conflicts, try to submit
 			const result = await this.submitChangelist(pending, submitRetries)
 			if (result.result !== 'error') {
 				return result
 			}
 
-			failure = { kind: 'Commit failure', description: result.message }
+			// todo: lock a target stream in the tests, print error message
+			// or lock a test stream and try to commit
+			// then look for error message in result.message
+			if (!this.options.approval) {
+				failure = { kind: 'Commit failure', description: result.message }
+			}
+		}
+
+		// no failure set means the change needs approval
+		if (failure) {
+			await this.handlePostIntegrationFailure(failure, pending)
+			this.sourceNode.handleMergeFailure(failure, pending, true)
 		}
+		else {
+			const approval = this.options.approval!
+
+			// shelve
+			await this.shelveChangelist(pending)
+
+			// still notify as a 'failure' to trigger normal blockage mechanism
+			failure = { kind: 'Approval required', description: approval.description }
 
-		await this.handlePostIntegrationFailure(failure, pending)
-		await this.sourceNode.handleMergeFailure(failure, pending, true)
+			// pause this bot
+			this.block(this.createEdgeBlockageInfo(failure, pending))
+			this.sourceNode.findOrCreateBlockage(failure, pending,
+				`Integration of CL#${pending.change.cl} to ${pending.action.branch.name} needs approval: ${approval.description}`,
+				{
+					group: approval.group,
+					shelfCl: pending.newCl
+				});
+		}
 
 		return new EdgeIntegrationDetails('error', `${failure.kind}: ${failure.description}`)
 	}
diff --git a/Engine/Extras/RoboMerge/v3/src/robo/nodebot.ts b/Engine/Extras/RoboMerge/v3/src/robo/nodebot.ts
--- a/Engine/Extras/RoboMerge/v3/src/robo/nodebot.ts
+++ b/Engine/Extras/RoboMerge/v3/src/robo/nodebot.ts
@@ -1832,30 +1832,34 @@ export class NodeBot extends PerforceStatefulBot implements NodeBotInterface {
 
 			message += '\n' + failure.summary 
 			this._sendError(new Recipients(owner), shortMessage, message)
+			return false
 		}
-		else {
-			// see if we've already sent an email for this one
-			const existingBlockage = this.conflicts.find(targetBranch, pending.change.source_cl)
-			const blockageOccurred = existingBlockage ? existingBlockage.time : new Date
 
-			const ownerEmail = this.p4.getEmail(owner)
-			const blockage: Blockage = {change: pending.change, action: pending.action, failure, owner, ownerEmail, time: blockageOccurred}
+		return this.findOrCreateBlockage(failure, pending, shortMessage)
+	}
 
-			if (existingBlockage) {
-				this.nodeBotLogger.info(shortMessage + ' (already notified)')
+	findOrCreateBlockage(failure: Failure, pending: PendingChange, shortMessage: string, approval?: {group: string, shelfCl: number}) {
+		const owner = getIntegrationOwner(pending) || pending.change.author
 
-				this.conflicts.updateBlockage(blockage)
-			}
-			else {
-				if (this.tickJournal && (failure.kind === 'Merge conflict' || failure.kind === 'Exclusive check-out')) {
-					++this.tickJournal.conflicts;
-				}
+		// see if we've already created a notification for this one
+		const existingBlockage = this.conflicts.find(pending.action.branch, pending.change.source_cl)
+		const blockageOccurred = existingBlockage ? existingBlockage.time : new Date
 
-				this.conflicts.onBlockage(blockage)
-				return true
-			}
+		const ownerEmail = this.p4.getEmail(owner)
+		const blockage: Blockage = {change: pending.change, action: pending.action, failure, owner, ownerEmail, time: blockageOccurred, approval}
+
+		if (existingBlockage) {
+			this.nodeBotLogger.info(shortMessage + ' (already notified)')
+
+			this.conflicts.updateBlockage(blockage)
+			return false
 		}
-		return false
+		if (this.tickJournal && (failure.kind === 'Merge conflict' || failure.kind === 'Exclusive check-out')) {
+			++this.tickJournal.conflicts;
+		}
+
+		this.conflicts.onBlockage(blockage)
+		return true
 	}
 
 	private sendNagEmails() {
diff --git a/Engine/Extras/RoboMerge/v3/src/robo/notifications.ts b/Engine/Extras/RoboMerge/v3/src/robo/notifications.ts
--- a/Engine/Extras/RoboMerge/v3/src/robo/notifications.ts
+++ b/Engine/Extras/RoboMerge/v3/src/robo/notifications.ts
@@ -397,7 +397,7 @@ export class BotNotifications implements BotEventHandler {
 			return
 		}
 
-		const cl = changeInfo.cl
+		const cl = blockage.approval ? blockage.approval.shelfCl : changeInfo.cl
 
 		// TODO: Support DMs if we don't have a channel configured
 		if (!this.slackMessages) {
@@ -420,9 +420,11 @@ export class BotNotifications implements BotEventHandler {
 		const channelPing = userEmail ? blockage.owner : `@${blockage.owner}`
 
 		const isBotUser = isUserAKnownBot(blockage.owner)
-		const text = isBotUser ? `Blockage caused by \`${blockage.owner}\` commit!` :
-			blockage.failure.kind === 'Too many files' ? `${channelPing}, please request a shelf for this large changelist` :
-			`${channelPing}, please resolve the following ${issue}:`
+		const text =
+			blockage.approval ?								`${channelPing}'s change needs to be approved by @${blockage.approval.group}` :
+			isBotUser ? 									`Blockage caused by \`${blockage.owner}\` commit!` :
+			blockage.failure.kind === 'Too many files' ?	`${channelPing}, please request a shelf for this large changelist` :
+															`${channelPing}, please resolve the following ${issue}:`
 
 		const message = this.makeSlackChannelMessage(
 			`${sourceBranch.name} blocked! (${issue})`,
diff --git a/Engine/Extras/RoboMerge/v3/src/robo/status-types.ts b/Engine/Extras/RoboMerge/v3/src/robo/status-types.ts
--- a/Engine/Extras/RoboMerge/v3/src/robo/status-types.ts
+++ b/Engine/Extras/RoboMerge/v3/src/robo/status-types.ts
@@ -24,6 +24,7 @@ export type FailureKind =
 	'Syntax error' |
 	'Disallowed files' |
 	'Too many files' |
+	'Approval required' |
 	'Conversion to edits failure' |
 	'Unit Test error'
 
