[NO NON-PII COMMIT MESSAGE]

diff --git a/Engine/Extras/RoboMerge/v3/src/robo/branch-interfaces.ts b/Engine/Extras/RoboMerge/v3/src/robo/branch-interfaces.ts
--- a/Engine/Extras/RoboMerge/v3/src/robo/branch-interfaces.ts
+++ b/Engine/Extras/RoboMerge/v3/src/robo/branch-interfaces.ts
@@ -4,7 +4,7 @@ import { BranchSpec, ConflictedResolveNFile, RoboWorkspace } from '../common/per
 // TODO: Remove Circular Dependency on bot-interfaces
 import { NodeBotInterface } from './bot-interfaces';
 import { FailureKind } from './status-types';
-import { BotConfig, BranchBase, EdgeOptions, NodeOptions, IntegrationWindowPane } from './branchdefs';
+import { BotConfig, BranchBase, ApprovalOptions, EdgeOptions, NodeOptions, IntegrationWindowPane } from './branchdefs';
 import { BlockageNodeOpUrls } from './roboserver';
 
 export type BranchArg = Branch | string
@@ -168,7 +168,7 @@ export interface Blockage {
 	owner: string
 	ownerEmail: Promise<string | null>
 	approval?: {
-		group: string
+		settings: ApprovalOptions
 		shelfCl: number
 	}
 	time: Date
diff --git a/Engine/Extras/RoboMerge/v3/src/robo/branchdefs.ts b/Engine/Extras/RoboMerge/v3/src/robo/branchdefs.ts
--- a/Engine/Extras/RoboMerge/v3/src/robo/branchdefs.ts
+++ b/Engine/Extras/RoboMerge/v3/src/robo/branchdefs.ts
@@ -139,13 +139,15 @@ const edgeOptionFieldsPrototype = {
 
 	approval: {
 		description: '',
-		group: ''
+		channelName: '',
+		channelId: ''
 	}
 }
 
 const edgeOptionFieldNames: ReadonlySet<string> = new Set(Object.keys(edgeOptionFieldsPrototype));
 
 type EdgeOptionFields = typeof edgeOptionFieldsPrototype
+export type ApprovalOptions = typeof edgeOptionFieldsPrototype.approval
 
 export type NodeOptions = Partial<NodeOptionFields>
 export type EdgeOptions = Partial<EdgeOptionFields>
@@ -425,6 +427,13 @@ export class BranchDefs {
 					}
 				}
 
+				if (edge.approval) {
+					// should be replaced by generic handling of objects in the prototype
+					if (!edge.approval.description || !edge.approval.channelName || !edge.approval.channelId) {
+						throw new Error(`Invalid approval settings for edge ${edge.from}->${edge.to}`)
+					}
+				}
+
 				validateCommonOptions(edge)
 			}
 		}
diff --git a/Engine/Extras/RoboMerge/v3/src/robo/branchgraph.ts b/Engine/Extras/RoboMerge/v3/src/robo/branchgraph.ts
--- a/Engine/Extras/RoboMerge/v3/src/robo/branchgraph.ts
+++ b/Engine/Extras/RoboMerge/v3/src/robo/branchgraph.ts
@@ -187,6 +187,10 @@ if (botname === '__TEST__') {
 					throw new Error(`No edge matching property ${edge.from} -> ${edge.to}`)
 				}
 
+				if (edge.approval) {
+					edge.additionalSlackChannel = edge.approval.channelId
+				}
+
 				source!.edgeProperties.set(targetName, edge)
 			}
 		}
diff --git a/Engine/Extras/RoboMerge/v3/src/robo/edgebot.ts b/Engine/Extras/RoboMerge/v3/src/robo/edgebot.ts
--- a/Engine/Extras/RoboMerge/v3/src/robo/edgebot.ts
+++ b/Engine/Extras/RoboMerge/v3/src/robo/edgebot.ts
@@ -646,7 +646,7 @@ class EdgeBotImpl extends PerforceStatefulBot {
 			this.sourceNode.findOrCreateBlockage(failure, pending,
 				`Integration of CL#${pending.change.cl} to ${pending.action.branch.name} needs approval: ${approval.description}`,
 				{
-					group: approval.group,
+					settings: approval,
 					shelfCl: pending.newCl
 				});
 		}
diff --git a/Engine/Extras/RoboMerge/v3/src/robo/nodebot.ts b/Engine/Extras/RoboMerge/v3/src/robo/nodebot.ts
--- a/Engine/Extras/RoboMerge/v3/src/robo/nodebot.ts
+++ b/Engine/Extras/RoboMerge/v3/src/robo/nodebot.ts
@@ -6,6 +6,7 @@ import { ContextualLogger } from '../common/logger';
 import { Mailer, MailParams, Recipients } from '../common/mailer';
 import { Change, ConflictedResolveNFile, PerforceContext, RoboWorkspace } from '../common/perforce';
 import { IPCControls, NodeBotInterface, QueuedChange, ReconsiderArgs } from './bot-interfaces';
+import { ApprovalOptions } from './branchdefs'
 import { BlockagePauseInfo, BranchStatus } from './status-types';
 import { AlreadyIntegrated, Blockage, Branch, BranchArg, BranchGraphInterface, ChangeInfo, EndIntegratingToGateEvent, Failure } from './branch-interfaces';
 import { MergeAction, OperationResult, PendingChange, resolveBranchArg, StompedRevision, StompVerification, StompVerificationFile }  from './branch-interfaces';
@@ -1838,7 +1839,7 @@ export class NodeBot extends PerforceStatefulBot implements NodeBotInterface {
 		return this.findOrCreateBlockage(failure, pending, shortMessage)
 	}
 
-	findOrCreateBlockage(failure: Failure, pending: PendingChange, shortMessage: string, approval?: {group: string, shelfCl: number}) {
+	findOrCreateBlockage(failure: Failure, pending: PendingChange, shortMessage: string, approval?: {settings: ApprovalOptions, shelfCl: number}) {
 		const owner = getIntegrationOwner(pending) || pending.change.author
 
 		// see if we've already created a notification for this one
diff --git a/Engine/Extras/RoboMerge/v3/src/robo/notifications.ts b/Engine/Extras/RoboMerge/v3/src/robo/notifications.ts
--- a/Engine/Extras/RoboMerge/v3/src/robo/notifications.ts
+++ b/Engine/Extras/RoboMerge/v3/src/robo/notifications.ts
@@ -397,8 +397,6 @@ export class BotNotifications implements BotEventHandler {
 			return
 		}
 
-		const cl = blockage.approval ? blockage.approval.shelfCl : changeInfo.cl
-
 		// TODO: Support DMs if we don't have a channel configured
 		if (!this.slackMessages) {
 			// doing nothing at the moment - don't want to complicate things with fallbacks
@@ -407,6 +405,8 @@ export class BotNotifications implements BotEventHandler {
 			return
 		}
 
+		const cl = blockage.approval ? blockage.approval.shelfCl : changeInfo.cl
+
 		// or integration failure (better wording? exclusive check-out?)
 		const sourceBranch = changeInfo.branch
 		let targetBranch
@@ -421,7 +421,7 @@ export class BotNotifications implements BotEventHandler {
 
 		const isBotUser = isUserAKnownBot(blockage.owner)
 		const text =
-			blockage.approval ?								`${channelPing}'s change needs to be approved by @${blockage.approval.group}` :
+			blockage.approval ?				`${channelPing}'s change needs to be approved in ${blockage.approval.settings.channelName}` :
 			isBotUser ? 									`Blockage caused by \`${blockage.owner}\` commit!` :
 			blockage.failure.kind === 'Too many files' ?	`${channelPing}, please request a shelf for this large changelist` :
 															`${channelPing}, please resolve the following ${issue}:`
@@ -456,19 +456,33 @@ export class BotNotifications implements BotEventHandler {
 
 		// Post message to owner in DM
 		if (DIRECT_MESSAGING_ENABLED && !isBotUser && targetBranch && userEmail) {
-			const dmText = `Your change (${makeClLink(changeInfo.source_cl)}) ` +
-				`hit '${issue}' while merging from *${sourceBranch.name}* to *${targetBranch.name}*.\n\n` +
-				'`' + blockage.change.description.substr(0, 80) + '`\n\n' +
-				"*_Resolving this blockage is time sensitive._ Please select one of the following:*"
-			
-			const urls = this.blockageUrlGenerator(blockage)
-			if (!urls) {
-				const error = `Could not get blockage URLs for blockage -- CL ${blockage.change.cl}`
-				this.botNotificationsLogger.printException(error)
-				throw error
+			const sourceLink = makeClLink(changeInfo.source_cl)
+			let dm: SlackMessage
+			if (blockage.approval) {
+				dm = {
+					title: 'Approval needed to commit to ' + targetBranch.name,
+					text: `Your change ${sourceLink} has been sent to ${blockage.approval.settings.channelName} for approval\n\n` +
+							blockage.approval.settings.description,
+					channel: "",
+					mrkdwn: false
+				}
+			}
+			else {
+				const dmText = `Your change (${sourceLink}) ` +
+					`hit '${issue}' while merging from *${sourceBranch.name}* to *${targetBranch.name}*.\n\n` +
+					'`' + blockage.change.description.substr(0, 80) + '`\n\n' +
+					"*_Resolving this blockage is time sensitive._ Please select one of the following:*"
+				
+				const urls = this.blockageUrlGenerator(blockage)
+				if (!urls) {
+					const error = `Could not get blockage URLs for blockage -- CL ${blockage.change.cl}`
+					this.botNotificationsLogger.printException(error)
+					throw error
+				}
+
+				dm = this.makeSlackDirectMessage(dmText, changeInfo.source_cl, cl, targetBranch.name, urls)
 			}
 
-			const dm = this.makeSlackDirectMessage(dmText, changeInfo.source_cl, cl, targetBranch.name, urls)
 			this.slackMessages.postDM(userEmail, changeInfo.source_cl, targetBranch, dm)
 		}
 	}
