[NO NON-PII COMMIT MESSAGE]

diff --git a/Engine/Extras/RoboMerge/v3/src/robo/notifications.ts b/Engine/Extras/RoboMerge/v3/src/robo/notifications.ts
--- a/Engine/Extras/RoboMerge/v3/src/robo/notifications.ts
+++ b/Engine/Extras/RoboMerge/v3/src/robo/notifications.ts
@@ -456,19 +456,18 @@ export class BotNotifications implements BotEventHandler {
 
 		// Post message to owner in DM
 		if (DIRECT_MESSAGING_ENABLED && !isBotUser && targetBranch && userEmail) {
-			const sourceLink = makeClLink(changeInfo.source_cl)
 			let dm: SlackMessage
 			if (blockage.approval) {
 				dm = {
 					title: 'Approval needed to commit to ' + targetBranch.name,
-					text: `Your change ${sourceLink} has been sent to ${blockage.approval.settings.channelName} for approval\n\n` +
+					text: `Your change has been shelved in ${makeClLink(cl)} and sent to ${blockage.approval.settings.channelName} for approval\n\n` +
 							blockage.approval.settings.description,
 					channel: "",
 					mrkdwn: false
 				}
 			}
 			else {
-				const dmText = `Your change (${sourceLink}) ` +
+				const dmText = `Your change (${makeClLink(changeInfo.source_cl)}) ` +
 					`hit '${issue}' while merging from *${sourceBranch.name}* to *${targetBranch.name}*.\n\n` +
 					'`' + blockage.change.description.substr(0, 80) + '`\n\n' +
 					"*_Resolving this blockage is time sensitive._ Please select one of the following:*"
