[NO NON-PII COMMIT MESSAGE]

diff --git a/Engine/Extras/RoboMerge/v3/functional_tests/framework/src/tests/stomp-forwarding-commands.ts b/Engine/Extras/RoboMerge/v3/functional_tests/framework/src/tests/stomp-forwarding-commands.ts
--- a/Engine/Extras/RoboMerge/v3/functional_tests/framework/src/tests/stomp-forwarding-commands.ts
+++ b/Engine/Extras/RoboMerge/v3/functional_tests/framework/src/tests/stomp-forwarding-commands.ts
@@ -14,8 +14,11 @@ export class StompForwardingCommands extends FunctionalTest {
 		this.mainClient = this.getClient('Main')
 		await P4Util.addFileAndSubmit(this.mainClient, 'test.uasset', 'Initial content', true)
 
+		// popluateStreams does all in parallel, ignoring dependencies
+		await this.populateStreams(streams.filter(s => s.name !== 'DevGrandchild'))
+		await this.populateStreams(streams.slice(2)) // Dev-Grandchild
 
-		await this.populateStreams(streams)
+		// await this.populateStreams(streams)
 
 		const childClient = this.getClient('DevChild')
 		await childClient.sync()
