[NO NON-PII COMMIT MESSAGE]

diff --git a/Engine/Extras/RoboMerge/v3/src/common/p4util.ts b/Engine/Extras/RoboMerge/v3/src/common/p4util.ts
--- a/Engine/Extras/RoboMerge/v3/src/common/p4util.ts
+++ b/Engine/Extras/RoboMerge/v3/src/common/p4util.ts
@@ -169,7 +169,7 @@ export async function convertIntegrateToEdit(p4: PerforceContext, roboWorkspace:
 	}
 }
 
-export async function cleanWorkspaces(p4: PerforceContext, workspaces: [string, string][]) {
+export async function cleanWorkspaces(p4: PerforceContext, workspaces: [string, string][], edgeServerAddress?: string) {
 	const changes = await p4.get_pending_changes() as Change[]
 
 	const nameSet = new Set(workspaces.map(([name, _]) => name))
@@ -215,5 +215,5 @@ export async function cleanWorkspaces(p4: PerforceContext, workspaces: [string,
 	}
 
 	p4utilsLogger.info('Resetting all workspaces to revision 0')
-	await Promise.all(workspaces.map(([name, root]) => p4.sync(name, root + '#0')))
+	await Promise.all(workspaces.map(([name, root]) => p4.sync(name, root + '#0', undefined, edgeServerAddress)))
 }
diff --git a/Engine/Extras/RoboMerge/v3/src/common/perforce.ts b/Engine/Extras/RoboMerge/v3/src/common/perforce.ts
--- a/Engine/Extras/RoboMerge/v3/src/common/perforce.ts
+++ b/Engine/Extras/RoboMerge/v3/src/common/perforce.ts
@@ -225,6 +225,10 @@ export interface Change {
 	isUserRequest?: boolean;
 	ignoreExcludedAuthors?: boolean
 	forceCreateAShelf?: boolean
+	edgeServerToHostShelf?: {
+		id: string
+		address: string
+	}
 	sendNoShelfEmail?: boolean // Used for requesting stomps for internal RM usage, such as stomp changes
 	commandOverride?: string
 
@@ -265,12 +269,13 @@ export function isExecP4Error(err: any): err is [Error, string] {
 }
 
 interface ExecOpts {
-	stdin?: string;
-	quiet?: boolean;
-	noCwd?: boolean;
-	noUsername?: boolean;
-	numRetries?: number;
-	trace?: boolean;
+	stdin?: string
+	quiet?: boolean
+	noCwd?: boolean
+	noUsername?: boolean
+	numRetries?: number
+	trace?: boolean
+	edgeServerAddress?: string
 }
 
 interface ExecZtagOpts extends ExecOpts {
@@ -280,6 +285,7 @@ interface ExecZtagOpts extends ExecOpts {
 export interface EditOwnerOpts {
 	newWorkspace?: string;
 	changeSubmitted?: boolean;
+	edgeServerAddress?: string
 }
 
 export interface ConflictedResolveNFile {
@@ -467,12 +473,17 @@ export class PerforceContext {
 	}
 
 	// get a list of all pending changes for this user
-	async get_pending_changes() {
+	async get_pending_changes(workspaceName?: string, edgeServerAddress?: string) {
 		if (!this.username) {
 			throw new Error("username not set");
 		}
 
-		return this._execP4Ztag(null, ['changes', '-u', this.username, '-s', 'pending'], { multiline: true });
+		let args = ['changes', '-u', this.username, '-s', 'pending']
+		if (workspaceName) {
+			args = [...args, '-c', workspaceName]
+		}
+
+		return this._execP4Ztag(null, args, { multiline: true, edgeServerAddress });
 	}
 
 	/** get a single change in the format of changes() */
@@ -548,9 +559,19 @@ export class PerforceContext {
 
 	// find a workspace for the given user
 	// output format is an array of workspace names
-	async find_workspaces(user?: string) {
+	async find_workspaces(user?: string, edgeServerAddress?: string) {
 		// -a to include workspaces on edge servers
-		let parsedClients = await this._execP4Ztag(null, ['clients', '-u', user || this.username, '-a'], { multiline: true });
+
+
+		let args = ['clients', '-u', user || this.username]
+
+		if (!edgeServerAddress) {
+			// find all
+			args.push('-a')
+		}
+
+
+		let parsedClients = await this._execP4Ztag(null, args, { multiline: true, edgeServerAddress });
 		let workspaces = [];
 		for (let clientDef of parsedClients) {
 			if (clientDef.client) {
@@ -560,8 +581,24 @@ export class PerforceContext {
 		return workspaces as ClientSpec[];
 	}
 
-	async find_workspace_by_name(workspaceName: string) {
-		return await this._execP4Ztag(null, ['clients', '-E', workspaceName]);
+	find_workspace_by_name(workspaceName: string) {
+		return this._execP4Ztag(null, ['clients', '-E', workspaceName]);
+	}
+
+	async getWorkspaceEdgeServer(workspaceName: string): Promise<{id: string, address: string} | null> {
+		const serverId = await this._execP4(null, ['-ztag', '-F', '%ServerID%', 'client', '-o', workspaceName]) 
+		if (serverId) {
+			const address = await this.getEdgeServerAddress(serverId)
+			if (!address) {
+				throw new Error(`Couldn't find address for edge server '${serverId}'`)
+			}
+			return {id: serverId, address}
+		}
+		return null
+	}
+
+	getEdgeServerAddress(serverId: string): Promise<string> {
+		return this._execP4(null, ['-ztag', '-F', '%Address%', 'server', '-o', serverId])
 	}
 
 	async clean(roboWorkspace: RoboWorkspace) {
@@ -586,11 +623,16 @@ export class PerforceContext {
 	}
 
 	// sync the depot path specified
-	async sync(roboWorkspace: RoboWorkspace, depotPath: string, opts?: string[]) {
+	async sync(roboWorkspace: RoboWorkspace, depotPath: string, opts?: string[], edgeServerAddress?: string) {
 		const workspace = coercePerforceWorkspace(roboWorkspace);
-		const args = ['sync', ...(opts || []), depotPath];
+		let args = edgeServerAddress ? ['-p', edgeServerAddress] : []
+		args.push('sync')
+		if (opts) {
+			args = [...args, ...opts]
+		}
+		args.push(depotPath)
 		try {
-			await this._execP4(workspace, args);
+			await this._execP4(workspace, args)
 		}
 		catch (reason) {
 			if (!isExecP4Error(reason)) {
@@ -599,7 +641,7 @@ export class PerforceContext {
 			const [err, output] = reason
 			// this is an acceptable non-error case for us
 			if (!output || typeof output !== "string" || !output.trim().endsWith("up-to-date."))
-				throw err;
+				throw err
 		}
 	}
 
@@ -613,7 +655,7 @@ export class PerforceContext {
 		return change.change;
 	}
 
-	async newWorkspace(workspaceName: string, params: any) {
+	async newWorkspace(workspaceName: string, params: any, edgeServerAddress?: string) {
 		params.Client = workspaceName
 		if (!('Root' in params)) {
 			params.Root = 'd:/ROBO/' + workspaceName // default windows path
@@ -651,12 +693,12 @@ export class PerforceContext {
 
 		// run the p4 client command
 		this.logger.info(`Executing: 'p4 client -i' to create workspace ${workspaceName}`);
-		return this._execP4(null, ['client', '-i'], { stdin: workspaceForm, quiet: true });
+		return this._execP4(null, ['client', '-i'], { stdin: workspaceForm, quiet: true, edgeServerAddress });
 	}
 
 	// Create a new workspace for Robomerge GraphBot
-	async newGraphBotWorkspace(name: string, extraParams: any) {
-		return this.newWorkspace(name, {Root: '/src/' + name, ...extraParams});
+	async newGraphBotWorkspace(name: string, extraParams: any, edgeServerAddress?: string) {
+		return this.newWorkspace(name, {Root: '/src/' + name, ...extraParams}, edgeServerAddress);
 	}
 
 	// Create a new workspace for Robomerge to read branchspecs from
@@ -685,7 +727,7 @@ export class PerforceContext {
 
 	// create a new CL with a specific description
 	// output format is just CL number
-	async new_cl(roboWorkspace: RoboWorkspace, description: string, files?: string[]) {
+	async new_cl(roboWorkspace: RoboWorkspace, description: string, files?: string[], edgeServerAddress?: string) {
 		const workspace = coercePerforceWorkspace(roboWorkspace);
 		// build the minimal form
 		let form = 'Change:\tnew\nStatus:\tnew\nType:\tpublic\n';
@@ -702,7 +744,7 @@ export class PerforceContext {
 
 		// run the P4 change command
 		this.logger.info("Executing: 'p4 change -i' to create a new CL");
-		const output = await this._execP4(workspace, ['change', '-i'], { stdin: form, quiet: true });
+		const output = await this._execP4(workspace, ['change', '-i'], { stdin: form, quiet: true, edgeServerAddress });
 		// parse the CL out of output
 		const match = output.match(/Change (\d+) created./);
 		if (!match) {
@@ -716,7 +758,8 @@ export class PerforceContext {
 	// integrate a CL from source to destination, resolve, and place the results in a new CL
 	// output format is true if the integration resolved or false if the integration wasn't necessary (still considered a success)
 	// failure to resolve is treated as an error condition
-	async integrate(roboWorkspace: RoboWorkspace, source: IntegrationSource, dest_changelist: number, target: IntegrationTarget): Promise<[string, (Change | string)[]]> {
+	async integrate(roboWorkspace: RoboWorkspace, source: IntegrationSource, dest_changelist: number, target: IntegrationTarget, edgeServerAddress?: string)
+		: Promise<[string, (Change | string)[]]> {
 		const workspace = coercePerforceWorkspace(roboWorkspace);
 		// build a command
 		let cmdList = [
@@ -758,7 +801,7 @@ export class PerforceContext {
 		// execute the P4 command
 		let changes;
 		try {
-			changes = await this._execP4Ztag(workspace, cmdList, { numRetries: 0 });
+			changes = await this._execP4Ztag(workspace, cmdList, { numRetries: 0, edgeServerAddress });
 		}
 		catch (reason) {
 			if (!isExecP4Error(reason)) {
@@ -817,10 +860,10 @@ export class PerforceContext {
 		return ["integrated", <Change[]>changes];
 	}
 
-	async resolveHelper(workspace: Workspace | null, flag: string, changelist: number): Promise<any> {
+	async resolveHelper(workspace: Workspace | null, flag: string, changelist: number, edgeServerAddress?: string): Promise<any> {
 		try {
 			// Perform merge
-			return await this._execP4Ztag(workspace, ['resolve', flag, `-c${changelist}`])
+			return await this._execP4Ztag(workspace, ['resolve', flag, `-c${changelist}`], {edgeServerAddress})
 		}
 		catch (reason) {
 			if (!isExecP4Error(reason)) {
@@ -850,7 +893,8 @@ export class PerforceContext {
 
 
 	// output format a list of conflicting files (p4 output)
-	async resolve(roboWorkspace: RoboWorkspace, changelist: number, resolution: string, disallowDashN?: boolean): Promise<ResolveResult> {
+	async resolve(roboWorkspace: RoboWorkspace, changelist: number, resolution: string, disallowDashN?: boolean, edgeServerAddress?: string)
+		: Promise<ResolveResult> {
 		const workspace = coercePerforceWorkspace(roboWorkspace)
 		let flag = null
 		switch (resolution) {
@@ -863,7 +907,7 @@ export class PerforceContext {
 		}
 
 		// Perform merge
-		let fileInfo = await this.resolveHelper(workspace, flag, changelist)
+		let fileInfo = await this.resolveHelper(workspace, flag, changelist, edgeServerAddress)
 
 		// Clobber remaining files after first merge, if requested
 		if (resolution === 'clobber') {
@@ -877,7 +921,7 @@ export class PerforceContext {
 			}
 
 			// Otherwise, fileInfo should be updated with the clobber result
-			fileInfo = await this.resolveHelper(workspace, '-at', changelist)
+			fileInfo = await this.resolveHelper(workspace, '-at', changelist, edgeServerAddress)
 		}
 
 		// If resolveHelper() returned a string, we do not need to return a dashNResult
@@ -893,7 +937,7 @@ export class PerforceContext {
 		let dashNresult: string[]
 		try {
 			//dashNresult = await this._execP4(workspace, ['resolve', '-N', `-c${changelist}`])
-			dashNresult = await this._execP4Ztag(workspace, ['resolve', '-N', `-c${changelist}`])
+			dashNresult = await this._execP4Ztag(workspace, ['resolve', '-N', `-c${changelist}`], {edgeServerAddress})
 		}
 		catch (reason) {
 			if (!isExecP4Error(reason)) {
@@ -998,9 +1042,9 @@ export class PerforceContext {
 
 	// delete a CL
 	// output format is just error or not
-	deleteCl(roboWorkspace: RoboWorkspace, changelist: number) {
+	deleteCl(roboWorkspace: RoboWorkspace, changelist: number, edgeServerAddress?: string) {
 		const workspace = coercePerforceWorkspace(roboWorkspace);
-		return this._execP4(workspace, ["change", "-d", changelist.toString()]);
+		return this._execP4(workspace, ["change", "-d", changelist.toString()], {edgeServerAddress});
 	}
 
 	// run p4 'opened' command: lists files in changelist with details of edit state (e.g. if a copy, provides source path)
@@ -1020,11 +1064,13 @@ export class PerforceContext {
 
 	// revert a CL deleting any files marked for add
 	// output format is just error or not
-	async revert(roboWorkspace: RoboWorkspace, changelist: number, opts?: string[]) {
+	async revert(roboWorkspace: RoboWorkspace, changelist: number, opts?: string[], edgeServerAddress?: string) {
 		const workspace = coercePerforceWorkspace(roboWorkspace);
 		const path = workspace && workspace.name ? `//${workspace.name}/...` : ''
 		try {
-			await this._execP4(workspace, ['revert', /*"-w", seems to cause ENOENT occasionally*/ ...(opts || []), '-c', changelist.toString(), path]);
+			await this._execP4(workspace, ['revert',
+				/*"-w", seems to cause ENOENT occasionally*/ ...(opts || []), '-c', changelist.toString(), path],
+				{edgeServerAddress});
 		}
 		catch (reason) {
 			if (!isExecP4Error(reason)) {
@@ -1071,10 +1117,11 @@ export class PerforceContext {
 
 	// shelve a CL
 	// output format is just error or not
-	async shelve(roboWorkspace: RoboWorkspace, changelist: number) {
+	async shelve(roboWorkspace: RoboWorkspace, changelist: number, edgeServerAddress?: string) {
 		const workspace = coercePerforceWorkspace(roboWorkspace)
 		try {
-			await this._execP4(workspace, ['shelve', '-f', '-c', changelist.toString()], { numRetries: 0 })
+			await this._execP4(workspace, ['shelve', '-f', '-c', changelist.toString()],
+				{ numRetries: 0, edgeServerAddress })
 		}
 		catch (reason) {
 			if (!isExecP4Error(reason)) {
@@ -1193,7 +1240,8 @@ export class PerforceContext {
 
 		opts = opts || {}; // optional newWorkspace:string and/or changeSubmitted:boolean
 		// get the current changelist description
-		const output = await this._execP4(workspace, ['change', '-o', changelist.toString()]);
+		const output = await this._execP4(workspace, ['change', '-o', changelist.toString()],
+			{edgeServerAddress: opts.edgeServerAddress});
 
 		// replace the description
 		let form = output.replace(/\nUser:\t[^\n]*\n/, `\nUser:\t${newOwner}\n`);
@@ -1205,7 +1253,8 @@ export class PerforceContext {
 		const changeFlag = opts.changeSubmitted ? '-f' : '-u';
 		this.logger.info(`Executing: 'p4 change -i ${changeFlag}' to edit user/client on CL${changelist}`);
 
-		await this._execP4(workspace, ['change', '-i', changeFlag], { stdin: form, quiet: true })
+		await this._execP4(workspace, ['change', '-i', changeFlag],
+				{ stdin: form, quiet: true, edgeServerAddress: opts.edgeServerAddress })
 	}
 
 	where(roboWorkspace: RoboWorkspace, clientPath: string) {
@@ -1268,15 +1317,19 @@ export class PerforceContext {
 		const workspace = coercePerforceWorkspace(roboWorkspace);
 		// add the client explicitly if one is set (should be done at call time)
 
-		const opts = optsIn || {};
+		const opts = optsIn || {}
 
 		// Perforce can get mighty confused if you log into multiple accounts at any point (only relevant to local debugging)
 		if (!opts.noUsername) {
-			args = ['-u', getPerforceUsername(), ...args];
+			args = ['-u', getPerforceUsername(), ...args]
 		}
 
 		if (workspace && workspace.name) {
-			args = ['-c', workspace.name, ...args];
+			args = ['-c', workspace.name, ...args]
+		}
+
+		if (opts.edgeServerAddress) {
+			args = ['-p', opts.edgeServerAddress, ...args]
 		}
 
 		args = ['-zprog=robomerge', '-zversion=' + robomergeVersion, ...args]
diff --git a/Engine/Extras/RoboMerge/v3/src/robo/branch-interfaces.ts b/Engine/Extras/RoboMerge/v3/src/robo/branch-interfaces.ts
--- a/Engine/Extras/RoboMerge/v3/src/robo/branch-interfaces.ts
+++ b/Engine/Extras/RoboMerge/v3/src/robo/branch-interfaces.ts
@@ -144,6 +144,10 @@ export interface ChangeInfo extends TargetInfo {
 
 	propagatingNullMerge: boolean
 	forceCreateAShelf: boolean
+	edgeServerToHostShelf?: {
+		id: string
+		address: string
+	}
 	overriddenCommand: string
 	macros: string[]
 	hasOkForGithubTag: boolean
diff --git a/Engine/Extras/RoboMerge/v3/src/robo/edgebot.ts b/Engine/Extras/RoboMerge/v3/src/robo/edgebot.ts
--- a/Engine/Extras/RoboMerge/v3/src/robo/edgebot.ts
+++ b/Engine/Extras/RoboMerge/v3/src/robo/edgebot.ts
@@ -1,14 +1,18 @@
 // Copyright Epic Games, Inc. All Rights Reserved.
 
+import * as fs from 'fs';
+import * as p4util from '../common/p4util';
+
 import { ContextualLogger } from "../common/logger";
 import { Recipients } from "../common/mailer";
-import { Change, ClientSpec, ConflictedResolveNFile, EditOwnerOpts, IntegrationSource }  from "../common/perforce";
+import { Change, ClientSpec, ConflictedResolveNFile, EditOwnerOpts, IntegrationSource } from "../common/perforce";
+import { coercePerforceWorkspace, getRootDirectoryForBranch } from "../common/perforce";
 import { IntegrationTarget, isExecP4Error, OpenedFileRecord, PerforceContext, EXCLUSIVE_CHECKOUT_REGEX } from "../common/perforce";
 import { VersionReader } from "../common/version";
 import { EdgeBotInterface, IPCControls, ReconsiderArgs } from "./bot-interfaces";
 import { PauseState } from "./state-interfaces";
 import { BlockagePauseInfo, BlockagePauseInfoMinimal } from "./status-types";
-import { AlreadyIntegrated, Branch, ChangeInfo, ConflictingFile, Failure, MergeAction, MergeMode, PendingChange } from "./branch-interfaces";
+import { AlreadyIntegrated, Branch, ChangeInfo, ConflictingFile, Failure, MergeAction, PendingChange } from "./branch-interfaces";
 import { PersistentConflict } from "./conflict-interfaces";
 import { BotEventTriggers } from "./events";
 import { NodeBot } from "./nodebot";
@@ -37,7 +41,7 @@ const JIRA_REGEX = /^\s*#jira\s+(.*)/i
 
 const MAX_CONFLICTS_TO_LIST = 5
 
-type ResolveResultDetail = 'quick' | 'detailed'
+// type ResolveResultDetail = 'quick' | 'detailed'
 
 type EdgeIntegrationResult = 'ok' | 'shelved' | 'nothing to do' | 'skipped' | 'error'
 
@@ -463,32 +467,65 @@ class EdgeBotImpl extends PerforceStatefulBot {
 	}
 
 	private async integrate(info: ChangeInfo, target: MergeAction) : Promise<EdgeIntegrationDetails> {
-		let to_integrate = info.cl
+		const to_integrate = info.cl
 		this._log_action(`Integrating CL ${to_integrate} to ${this.targetBranch.name}`)
 
 		// if required, add author review here so they're not in target.description, which is used for shelf description in case of conflict
-		let desc = target.description! // target.description always ends in newline
+		const desc = target.description! // target.description always ends in newline
+
+		let targetWorkspace = coercePerforceWorkspace(this.targetBranch.workspace)!.name
+		const edgeServer = info.edgeServerToHostShelf
+		if (edgeServer) {
+
+			targetWorkspace += '_' + edgeServer.id.toUpperCase()
+
+			// make sure target workspace/directory exists (always reset for now)
+
+			// ensure local directory exists (_initWorkspacesForGraphBot)
+			const path = getRootDirectoryForBranch(targetWorkspace);
+			if (!fs.existsSync(path)) {
+				this.edgeBotLogger.info(`Making directory ${path}`);
+				fs.mkdirSync(path);
+			}
+
+			// do we already have a workspace? (_initBranchWorkspacesForAllBots, _getExistingWorkspaces)
+			const existingWorkspaceInfos = await this.p4.find_workspaces(undefined, edgeServer.address)
+			if (existingWorkspaceInfos.map(ws => ws.client).indexOf(targetWorkspace) >= 0) {
+				await p4util.cleanWorkspaces(this.p4, [[targetWorkspace, target.branch.rootPath]], edgeServer.address)
+			}
+			else {
+				// create one 
+				if (!target.branch.stream) {
+					throw new Error('only stream workspaces supported on edge servers')
+				}
+				await this.p4.newGraphBotWorkspace(targetWorkspace, {Stream: target.branch.stream})
+
+				// _initWorkspacesForGraphBot does clean-up stuff here, but I don't think it's necessary
+			}
+		}
+
+		const edgeServerAddress: string | undefined = edgeServer && edgeServer.address
 
 		// create a new CL
-		const changenum = await this.p4.new_cl(this.targetBranch.workspace, desc)
+		const changenum = await this.p4.new_cl(this.targetBranch.workspace, desc, undefined, edgeServerAddress)
 
 		// try to integrate
 		const branchSpecToTarget = this.sourceBranch.branchspec.get(this.targetBranch.upperName)
-		let source : IntegrationSource = {
+		const source : IntegrationSource = {
 			branchspec: branchSpecToTarget,
 			changelist: to_integrate,
 			depot: this.sourceBranch.depot,
 			path_from: this.sourceBranch.rootPath,
 			stream: this.sourceBranch.stream
 		}
-		let integTarget : IntegrationTarget = {
+		const integTarget : IntegrationTarget = {
 			depot: this.targetBranch.depot,
 			path_to: this.targetBranch.rootPath,
 			stream: this.targetBranch.stream
 		}
 
 		this.currentIntegrationStartTimestamp = Date.now()
-		const [mode, results] = await this.p4.integrate(this.targetBranch.workspace, source, changenum, integTarget)
+		const [mode, results] = await this.p4.integrate(this.targetBranch.workspace, source, changenum, integTarget, edgeServerAddress)
 
 		const pending: PendingChange = {change: info, action: target, newCl: changenum}
 
@@ -567,11 +604,6 @@ class EdgeBotImpl extends PerforceStatefulBot {
 		}
 	}
 
-	private performResolve(changelist: number, mergeMode: MergeMode, detail: ResolveResultDetail) {
-		return this.p4.resolve(this.targetBranch.workspace, changelist, mergeMode, detail === 'quick')
-	}
-
-
 	/**
 	 * Perform a resolve of the change described by pending
 	 * @param {PendingChange} pending Pending change to resolve and submit
@@ -580,8 +612,16 @@ class EdgeBotImpl extends PerforceStatefulBot {
 	async _resolveChangelist(pending: PendingChange, submitRetries  = 3) : Promise<EdgeIntegrationDetails> {
 		// do a resolve with P4
 		this._log_action(`Resolving CL ${pending.change.cl} against ${this.targetBranch.name}`)
-		let detail: ResolveResultDetail = 'detailed'
-		const result = await this.performResolve(pending.newCl, pending.action.mergeMode, detail)
+//		const detail: ResolveResultDetail = 'detailed'
+
+		const result = await this.p4.resolve(
+			this.targetBranch.workspace,
+			pending.newCl,
+			pending.action.mergeMode,
+			false, // detail === 'quick',
+			pending.change.edgeServerToHostShelf && pending.change.edgeServerToHostShelf.address
+			)
+
 		if (pending.action.flags.has('manual') || pending.change.forceCreateAShelf) {
 			// the user requested manual merge
 			await this.shelveChangelist(pending)
@@ -752,12 +792,14 @@ class EdgeBotImpl extends PerforceStatefulBot {
 			final_desc += pending.change.additionalDescriptionText
 		}
 
+		const edgeServerAddress = pending.change.edgeServerToHostShelf && pending.change.edgeServerToHostShelf.address
+
 		// abort shelve if this is a buildmachine / robomerge change (unless we are forcing the shelf)
 		let failed = false
 		if (isUserAKnownBot(owner) && !pending.change.forceCreateAShelf) {
 			// revert the changes locally
 			this._log_action(`Reverting shelf due to '${owner}' being a known bot`)
-			await this.p4.revert(this.targetBranch.workspace, changenum)
+			await this.p4.revert(this.targetBranch.workspace, changenum, undefined, edgeServerAddress)
 			failed = true
 		}
 		else {
@@ -765,7 +807,7 @@ class EdgeBotImpl extends PerforceStatefulBot {
 			await this.p4.editDescription(this.targetBranch.workspace, changenum, final_desc)
 
 			// shelve the files as we see them (should trigger a codereview email)
-			if (!await this.p4.shelve(this.targetBranch.workspace, changenum)) {
+			if (!await this.p4.shelve(this.targetBranch.workspace, changenum, edgeServerAddress)) {
 				failed = true
 			}
 		}
@@ -774,7 +816,7 @@ class EdgeBotImpl extends PerforceStatefulBot {
 			// abort abort!
 			// delete the cl
 			this._log_action(`Deleting CL ${changenum}`)
-			await this.p4.deleteCl(this.targetBranch.workspace, changenum)
+			await this.p4.deleteCl(this.targetBranch.workspace, changenum, edgeServerAddress)
 
 			pending.newCl = -1
 
@@ -784,7 +826,7 @@ class EdgeBotImpl extends PerforceStatefulBot {
 
 		// revert the changes locally
 		this._log_action(`Reverting CL ${changenum} locally. (conflict owner: ${owner})`)
-		await this.p4.revert(this.targetBranch.workspace, changenum)
+		await this.p4.revert(this.targetBranch.workspace, changenum, [], edgeServerAddress)
 
 		// figure out what workspace to put it in
 		const branch_stream = this.targetBranch.stream ? this.targetBranch.stream.toLowerCase() : null
@@ -823,7 +865,7 @@ class EdgeBotImpl extends PerforceStatefulBot {
 		}
 
 		// log if we couldn't find a workspace
-		const opts: EditOwnerOpts = {}
+		const opts: EditOwnerOpts = {edgeServerAddress}
 
 		if (targetWorkspace) {
 			this.edgeBotLogger.info(`Moving CL ${changenum} to workspace '${targetWorkspace}'`)
diff --git a/Engine/Extras/RoboMerge/v3/src/robo/nodebot.ts b/Engine/Extras/RoboMerge/v3/src/robo/nodebot.ts
--- a/Engine/Extras/RoboMerge/v3/src/robo/nodebot.ts
+++ b/Engine/Extras/RoboMerge/v3/src/robo/nodebot.ts
@@ -4,7 +4,7 @@ import * as Sentry from '@sentry/node';
 import { _nextTick } from '../common/helper';
 import { ContextualLogger } from '../common/logger';
 import { Mailer, MailParams, Recipients } from '../common/mailer';
-import { Change, ConflictedResolveNFile, PerforceContext, RoboWorkspace } from '../common/perforce';
+import { Change, ConflictedResolveNFile, PerforceContext, RoboWorkspace, coercePerforceWorkspace } from '../common/perforce';
 import { IPCControls, NodeBotInterface, QueuedChange, ReconsiderArgs } from './bot-interfaces';
 import { ApprovalOptions } from './branchdefs'
 import { BlockagePauseInfo, BranchStatus } from './status-types';
@@ -1332,12 +1332,23 @@ export class NodeBot extends PerforceStatefulBot implements NodeBotInterface {
 		optOwnerOverride?: string, optWorkspaceOverride?: RoboWorkspace, optTargetBranch?: Branch | null
 	) : Promise<EdgeMergeResults> {
 		const result = await this._createChangeInfo(change, availableEdges, optOwnerOverride, optWorkspaceOverride, optTargetBranch)
-		if (result.info) {
-			return await this._mergeClViaEdges(availableEdges, result.info, ignoreEdgePauseState)
+		if (!result.info) {
+			const emResult = new EdgeMergeResults()
+			emResult.errors = result.errors
+			return emResult
 		}
-		const emResult = new EdgeMergeResults()
-		emResult.errors = result.errors
-		return emResult
+
+		// should also do this for #manual changes (set up some testing around those first)
+		if (change.forceCreateAShelf && optWorkspaceOverride) {
+			// see if we need to talk to an edge server to create the shelf
+			const edgeServer = await this.p4.getWorkspaceEdgeServer(
+				coercePerforceWorkspace(optWorkspaceOverride)!.name)
+			if (edgeServer) {
+				change.edgeServerToHostShelf = edgeServer
+			}
+		}
+
+		return await this._mergeClViaEdges(availableEdges, result.info, ignoreEdgePauseState)
 	}
 
 	private async _createChangeInfo(
@@ -1582,6 +1593,7 @@ export class NodeBot extends PerforceStatefulBot implements NodeBotInterface {
 		// By default, we should perform any submit that comes from a change.
 		// If the change is setup to be forcibly shelved, honor it
 		let forceCreateAShelf = !!change.forceCreateAShelf
+		let edgeServerToHostShelf = change.edgeServerToHostShelf
 		let sendNoShelfEmail = !!change.sendNoShelfEmail
 
 		let forceStompChanges = !!change.forceStompChanges
@@ -1595,8 +1607,7 @@ export class NodeBot extends PerforceStatefulBot implements NodeBotInterface {
 			source: parsedLines.source,
 			description,
 			propagatingNullMerge: parsedLines.propagatingNullMerge,
-			forceCreateAShelf,
-			sendNoShelfEmail,
+			forceCreateAShelf, edgeServerToHostShelf, sendNoShelfEmail,
 
 			forceStompChanges, additionalDescriptionText,
 			hasOkForGithubTag: parsedLines.hasOkForGithubTag,
