[NO NON-PII COMMIT MESSAGE]

diff --git a/Engine/Extras/RoboMerge/v3/src/robo/graphbot.ts b/Engine/Extras/RoboMerge/v3/src/robo/graphbot.ts
--- a/Engine/Extras/RoboMerge/v3/src/robo/graphbot.ts
+++ b/Engine/Extras/RoboMerge/v3/src/robo/graphbot.ts
@@ -206,12 +206,46 @@ export class GraphBot implements GraphInterface, BotEventHandler {
 				if (!request) {
 					break
 				}
-				await node.processQueuedChange(request)
+				try {
+					await node.processQueuedChange(request)
+				}
+				catch(err) {
+					this.handleNodebotError(node, err)
+					return
+				}
 				node.persistQueuedChanges()
 			}
 		}
 	}
 
+	// mostly for nodebots, but could also be the auto reloader bot
+	private handleNodebotError(bot: Bot, err: Error) {
+		this._runningBots = false
+
+		let errStr = err.toString()
+		if (errStr.length > MAX_ERROR_LENGTH_TO_REPORT) {
+			errStr = errStr.substr(0, MAX_ERROR_LENGTH_TO_REPORT) + ` ... (error length ${errStr.length})`
+		}
+		else {
+			errStr += err.stack
+		}
+		this.lastError = {
+			nodeBot: bot.fullName,
+			error: errStr
+		}
+
+		Sentry.withScope((scope) => {
+			scope.setTag('graphBot', this.branchGraph.botname)
+			scope.setTag('nodeBot', bot.fullName)
+			scope.setTag('lastCl', bot.lastCl.toString())
+
+			Sentry.captureException(err);
+		})
+		const msg = `${this.lastError.nodeBot} fell over with error`
+		this.botLogger.printException(err, msg)
+		postToRobomergeAlerts(`@here ${msg}:\n\`\`\`${errStr}\`\`\``)
+	}
+
 	private async startBotsAsync() {
 		if (!this.waitTime) {
 			throw new Error('runbots must be called before startBots')
@@ -243,31 +277,7 @@ export class GraphBot implements GraphInterface, BotEventHandler {
 					ticked = await bot.tick()
 				}
 				catch (err) {
-					this._runningBots = false
-
-					let errStr = err.toString()
-					if (errStr.length > MAX_ERROR_LENGTH_TO_REPORT) {
-						errStr = errStr.substr(0, MAX_ERROR_LENGTH_TO_REPORT) + ` ... (error length ${errStr.length})`
-					}
-					else {
-						errStr += err.stack
-					}
-					this.lastError = {
-						nodeBot: bot.fullName,
-						error: errStr
-					}
-
-					Sentry.withScope((scope) => {
-						scope.setTag('graphBot', this.branchGraph.botname)
-						scope.setTag('nodeBot', bot.fullName)
-						scope.setTag('lastCl', bot.lastCl.toString())
-
-						Sentry.captureException(err);
-					})
-					const msg = `${this.lastError.nodeBot} fell over with error`
-					this.botLogger.printException(err, msg)
-					postToRobomergeAlerts(`@here ${msg}:\n\`\`\`${errStr}\`\`\``)
-
+					this.handleNodebotError(bot, err)
 					return
 				}
 				bot.isActive = false
