[NO NON-PII COMMIT MESSAGE]

diff --git a/Engine/Extras/RoboMerge/v3/src/common/p4util.ts b/Engine/Extras/RoboMerge/v3/src/common/p4util.ts
--- a/Engine/Extras/RoboMerge/v3/src/common/p4util.ts
+++ b/Engine/Extras/RoboMerge/v3/src/common/p4util.ts
@@ -204,14 +204,14 @@ export async function cleanWorkspaces(p4: PerforceContext, workspaces: [string,
 		}
 		p4utilsLogger.info(`Attempting to revert ${changeStr}`)
 		try {
-			await p4.revert(workspace, change.change)
+			await p4.revert(workspace, change.change, [], edgeServerAddress)
 		}
 		catch (err) {
 			// ignore revert errors on startup (As long as delete works, we're good)
 			p4utilsLogger.error(`Revert failed. Will try delete anyway: ${err}`)
 		}
 
-		await p4.deleteCl(workspace, change.change)
+		await p4.deleteCl(workspace, change.change, edgeServerAddress)
 	}
 
 	p4utilsLogger.info('Resetting all workspaces to revision 0')
diff --git a/Engine/Extras/RoboMerge/v3/src/common/perforce.ts b/Engine/Extras/RoboMerge/v3/src/common/perforce.ts
--- a/Engine/Extras/RoboMerge/v3/src/common/perforce.ts
+++ b/Engine/Extras/RoboMerge/v3/src/common/perforce.ts
@@ -225,10 +225,6 @@ export interface Change {
 	isUserRequest?: boolean;
 	ignoreExcludedAuthors?: boolean
 	forceCreateAShelf?: boolean
-	edgeServerToHostShelf?: {
-		id: string
-		address: string
-	}
 	sendNoShelfEmail?: boolean // Used for requesting stomps for internal RM usage, such as stomp changes
 	commandOverride?: string
 
@@ -586,13 +582,14 @@ export class PerforceContext {
 	}
 
 	async getWorkspaceEdgeServer(workspaceName: string): Promise<{id: string, address: string} | null> {
-		const serverId = await this._execP4(null, ['-ztag', '-F', '%ServerID%', 'client', '-o', workspaceName]) 
-		if (serverId) {
+		const serverIdLine = await this._execP4(null, ['-ztag', '-F', '%ServerID%', 'client', '-o', workspaceName]) 
+		if (serverIdLine) {
+			const serverId = serverIdLine.trim()
 			const address = await this.getEdgeServerAddress(serverId)
 			if (!address) {
 				throw new Error(`Couldn't find address for edge server '${serverId}'`)
 			}
-			return {id: serverId, address}
+			return {id: serverId, address: address.trim()}
 		}
 		return null
 	}
@@ -655,7 +652,7 @@ export class PerforceContext {
 		return change.change;
 	}
 
-	async newWorkspace(workspaceName: string, params: any, edgeServerAddress?: string) {
+	async newWorkspace(workspaceName: string, params: any, edgeServer?: {id: string, address: string}) {
 		params.Client = workspaceName
 		if (!('Root' in params)) {
 			params.Root = 'd:/ROBO/' + workspaceName // default windows path
@@ -677,6 +674,11 @@ export class PerforceContext {
 			params.LineEnd = 'local'
 		}
 
+		let args = ['client', '-i']
+		if (edgeServer) {
+			args = ['-p', edgeServer.address, ...args, `--serverid=${edgeServer.id}`]
+		}
+
 		let workspaceForm = ''
 		for (let key in params) {
 			let val = params[key];
@@ -693,12 +695,12 @@ export class PerforceContext {
 
 		// run the p4 client command
 		this.logger.info(`Executing: 'p4 client -i' to create workspace ${workspaceName}`);
-		return this._execP4(null, ['client', '-i'], { stdin: workspaceForm, quiet: true, edgeServerAddress });
+		return this._execP4(null, args, { stdin: workspaceForm, quiet: true });
 	}
 
 	// Create a new workspace for Robomerge GraphBot
-	async newGraphBotWorkspace(name: string, extraParams: any, edgeServerAddress?: string) {
-		return this.newWorkspace(name, {Root: '/src/' + name, ...extraParams}, edgeServerAddress);
+	async newGraphBotWorkspace(name: string, extraParams: any, edgeServer?: {id: string, address: string}) {
+		return this.newWorkspace(name, {Root: '/src/' + name, ...extraParams}, edgeServer);
 	}
 
 	// Create a new workspace for Robomerge to read branchspecs from
diff --git a/Engine/Extras/RoboMerge/v3/src/robo/edgebot.ts b/Engine/Extras/RoboMerge/v3/src/robo/edgebot.ts
--- a/Engine/Extras/RoboMerge/v3/src/robo/edgebot.ts
+++ b/Engine/Extras/RoboMerge/v3/src/robo/edgebot.ts
@@ -498,7 +498,7 @@ class EdgeBotImpl extends PerforceStatefulBot {
 				if (!target.branch.stream) {
 					throw new Error('only stream workspaces supported on edge servers')
 				}
-				await this.p4.newGraphBotWorkspace(targetWorkspace, {Stream: target.branch.stream})
+				await this.p4.newGraphBotWorkspace(targetWorkspace, {Stream: target.branch.stream}, edgeServer)
 
 				// _initWorkspacesForGraphBot does clean-up stuff here, but I don't think it's necessary
 			}
@@ -507,7 +507,7 @@ class EdgeBotImpl extends PerforceStatefulBot {
 		const edgeServerAddress: string | undefined = edgeServer && edgeServer.address
 
 		// create a new CL
-		const changenum = await this.p4.new_cl(this.targetBranch.workspace, desc, undefined, edgeServerAddress)
+		const changenum = await this.p4.new_cl(targetWorkspace, desc, undefined, edgeServerAddress)
 
 		// try to integrate
 		const branchSpecToTarget = this.sourceBranch.branchspec.get(this.targetBranch.upperName)
@@ -525,7 +525,7 @@ class EdgeBotImpl extends PerforceStatefulBot {
 		}
 
 		this.currentIntegrationStartTimestamp = Date.now()
-		const [mode, results] = await this.p4.integrate(this.targetBranch.workspace, source, changenum, integTarget, edgeServerAddress)
+		const [mode, results] = await this.p4.integrate(targetWorkspace, source, changenum, integTarget, edgeServerAddress)
 
 		const pending: PendingChange = {change: info, action: target, newCl: changenum}
 
@@ -543,7 +543,7 @@ class EdgeBotImpl extends PerforceStatefulBot {
 
 				// integration not necessary
 				this.edgeBotLogger.info(msg)
-				await this.p4.deleteCl(this.targetBranch.workspace, changenum)
+				await this.p4.deleteCl(targetWorkspace, changenum, edgeServerAddress)
 				return new EdgeIntegrationDetails('nothing to do', msg)
 			}
 		}
diff --git a/Engine/Extras/RoboMerge/v3/src/robo/nodebot.ts b/Engine/Extras/RoboMerge/v3/src/robo/nodebot.ts
--- a/Engine/Extras/RoboMerge/v3/src/robo/nodebot.ts
+++ b/Engine/Extras/RoboMerge/v3/src/robo/nodebot.ts
@@ -1344,7 +1344,7 @@ export class NodeBot extends PerforceStatefulBot implements NodeBotInterface {
 			const edgeServer = await this.p4.getWorkspaceEdgeServer(
 				coercePerforceWorkspace(optWorkspaceOverride)!.name)
 			if (edgeServer) {
-				change.edgeServerToHostShelf = edgeServer
+				result.info.edgeServerToHostShelf = edgeServer
 			}
 		}
 
@@ -1593,7 +1593,6 @@ export class NodeBot extends PerforceStatefulBot implements NodeBotInterface {
 		// By default, we should perform any submit that comes from a change.
 		// If the change is setup to be forcibly shelved, honor it
 		let forceCreateAShelf = !!change.forceCreateAShelf
-		let edgeServerToHostShelf = change.edgeServerToHostShelf
 		let sendNoShelfEmail = !!change.sendNoShelfEmail
 
 		let forceStompChanges = !!change.forceStompChanges
@@ -1607,7 +1606,7 @@ export class NodeBot extends PerforceStatefulBot implements NodeBotInterface {
 			source: parsedLines.source,
 			description,
 			propagatingNullMerge: parsedLines.propagatingNullMerge,
-			forceCreateAShelf, edgeServerToHostShelf, sendNoShelfEmail,
+			forceCreateAShelf, sendNoShelfEmail,
 
 			forceStompChanges, additionalDescriptionText,
 			hasOkForGithubTag: parsedLines.hasOkForGithubTag,
