[NO NON-PII COMMIT MESSAGE]

diff --git a/Engine/Extras/RoboMerge/v3/src/robo/branch-interfaces.ts b/Engine/Extras/RoboMerge/v3/src/robo/branch-interfaces.ts
--- a/Engine/Extras/RoboMerge/v3/src/robo/branch-interfaces.ts
+++ b/Engine/Extras/RoboMerge/v3/src/robo/branch-interfaces.ts
@@ -148,6 +148,7 @@ export interface ChangeInfo extends TargetInfo {
 		id: string
 		address: string
 	}
+	targetWorkspaceOverride?: string
 	overriddenCommand: string
 	macros: string[]
 	hasOkForGithubTag: boolean
diff --git a/Engine/Extras/RoboMerge/v3/src/robo/edgebot.ts b/Engine/Extras/RoboMerge/v3/src/robo/edgebot.ts
--- a/Engine/Extras/RoboMerge/v3/src/robo/edgebot.ts
+++ b/Engine/Extras/RoboMerge/v3/src/robo/edgebot.ts
@@ -473,16 +473,16 @@ class EdgeBotImpl extends PerforceStatefulBot {
 		// if required, add author review here so they're not in target.description, which is used for shelf description in case of conflict
 		const desc = target.description! // target.description always ends in newline
 
-		let targetWorkspace = coercePerforceWorkspace(this.targetBranch.workspace)!.name
+		info.targetWorkspaceOverride = coercePerforceWorkspace(this.targetBranch.workspace)!.name
 		const edgeServer = info.edgeServerToHostShelf
 		if (edgeServer) {
 
-			targetWorkspace += '_' + edgeServer.id.toUpperCase()
+			info.targetWorkspaceOverride += '_' + edgeServer.id.toUpperCase()
 
 			// make sure target workspace/directory exists (always reset for now)
 
 			// ensure local directory exists (_initWorkspacesForGraphBot)
-			const path = getRootDirectoryForBranch(targetWorkspace);
+			const path = getRootDirectoryForBranch(info.targetWorkspaceOverride);
 			if (!fs.existsSync(path)) {
 				this.edgeBotLogger.info(`Making directory ${path}`);
 				fs.mkdirSync(path);
@@ -490,15 +490,15 @@ class EdgeBotImpl extends PerforceStatefulBot {
 
 			// do we already have a workspace? (_initBranchWorkspacesForAllBots, _getExistingWorkspaces)
 			const existingWorkspaceInfos = await this.p4.find_workspaces(undefined, edgeServer.address)
-			if (existingWorkspaceInfos.map(ws => ws.client).indexOf(targetWorkspace) >= 0) {
-				await p4util.cleanWorkspaces(this.p4, [[targetWorkspace, target.branch.rootPath]], edgeServer.address)
+			if (existingWorkspaceInfos.map(ws => ws.client).indexOf(info.targetWorkspaceOverride) >= 0) {
+				await p4util.cleanWorkspaces(this.p4, [[info.targetWorkspaceOverride, target.branch.rootPath]], edgeServer.address)
 			}
 			else {
 				// create one 
 				if (!target.branch.stream) {
 					throw new Error('only stream workspaces supported on edge servers')
 				}
-				await this.p4.newGraphBotWorkspace(targetWorkspace, {Stream: target.branch.stream}, edgeServer)
+				await this.p4.newGraphBotWorkspace(info.targetWorkspaceOverride, {Stream: target.branch.stream}, edgeServer)
 
 				// _initWorkspacesForGraphBot does clean-up stuff here, but I don't think it's necessary
 			}
@@ -507,7 +507,7 @@ class EdgeBotImpl extends PerforceStatefulBot {
 		const edgeServerAddress: string | undefined = edgeServer && edgeServer.address
 
 		// create a new CL
-		const changenum = await this.p4.new_cl(targetWorkspace, desc, undefined, edgeServerAddress)
+		const changenum = await this.p4.new_cl(info.targetWorkspaceOverride, desc, undefined, edgeServerAddress)
 
 		// try to integrate
 		const branchSpecToTarget = this.sourceBranch.branchspec.get(this.targetBranch.upperName)
@@ -525,7 +525,7 @@ class EdgeBotImpl extends PerforceStatefulBot {
 		}
 
 		this.currentIntegrationStartTimestamp = Date.now()
-		const [mode, results] = await this.p4.integrate(targetWorkspace, source, changenum, integTarget, edgeServerAddress)
+		const [mode, results] = await this.p4.integrate(info.targetWorkspaceOverride, source, changenum, integTarget, edgeServerAddress)
 
 		const pending: PendingChange = {change: info, action: target, newCl: changenum}
 
@@ -543,7 +543,7 @@ class EdgeBotImpl extends PerforceStatefulBot {
 
 				// integration not necessary
 				this.edgeBotLogger.info(msg)
-				await this.p4.deleteCl(targetWorkspace, changenum, edgeServerAddress)
+				await this.p4.deleteCl(info.targetWorkspaceOverride, changenum, edgeServerAddress)
 				return new EdgeIntegrationDetails('nothing to do', msg)
 			}
 		}
@@ -615,7 +615,7 @@ class EdgeBotImpl extends PerforceStatefulBot {
 //		const detail: ResolveResultDetail = 'detailed'
 
 		const result = await this.p4.resolve(
-			this.targetBranch.workspace,
+			pending.change.targetWorkspaceOverride || this.targetBranch.workspace,
 			pending.newCl,
 			pending.action.mergeMode,
 			false, // detail === 'quick',
@@ -793,21 +793,22 @@ class EdgeBotImpl extends PerforceStatefulBot {
 		}
 
 		const edgeServerAddress = pending.change.edgeServerToHostShelf && pending.change.edgeServerToHostShelf.address
+		const destRoboWorkspace = pending.change.targetWorkspaceOverride || this.targetBranch.workspace
 
 		// abort shelve if this is a buildmachine / robomerge change (unless we are forcing the shelf)
 		let failed = false
 		if (isUserAKnownBot(owner) && !pending.change.forceCreateAShelf) {
 			// revert the changes locally
 			this._log_action(`Reverting shelf due to '${owner}' being a known bot`)
-			await this.p4.revert(this.targetBranch.workspace, changenum, undefined, edgeServerAddress)
+			await this.p4.revert(destRoboWorkspace, changenum, undefined, edgeServerAddress)
 			failed = true
 		}
 		else {
 			// edit the CL description
-			await this.p4.editDescription(this.targetBranch.workspace, changenum, final_desc)
+			await this.p4.editDescription(destRoboWorkspace, changenum, final_desc)
 
 			// shelve the files as we see them (should trigger a codereview email)
-			if (!await this.p4.shelve(this.targetBranch.workspace, changenum, edgeServerAddress)) {
+			if (!await this.p4.shelve(destRoboWorkspace, changenum, edgeServerAddress)) {
 				failed = true
 			}
 		}
@@ -816,7 +817,7 @@ class EdgeBotImpl extends PerforceStatefulBot {
 			// abort abort!
 			// delete the cl
 			this._log_action(`Deleting CL ${changenum}`)
-			await this.p4.deleteCl(this.targetBranch.workspace, changenum, edgeServerAddress)
+			await this.p4.deleteCl(destRoboWorkspace, changenum, edgeServerAddress)
 
 			pending.newCl = -1
 
@@ -826,7 +827,7 @@ class EdgeBotImpl extends PerforceStatefulBot {
 
 		// revert the changes locally
 		this._log_action(`Reverting CL ${changenum} locally. (conflict owner: ${owner})`)
-		await this.p4.revert(this.targetBranch.workspace, changenum, [], edgeServerAddress)
+		await this.p4.revert(destRoboWorkspace, changenum, [], edgeServerAddress)
 
 		// figure out what workspace to put it in
 		const branch_stream = this.targetBranch.stream ? this.targetBranch.stream.toLowerCase() : null
@@ -876,7 +877,7 @@ class EdgeBotImpl extends PerforceStatefulBot {
 		}
 
 		// edit the owner to the author so they can resolve and submit themselves
-		await this.p4.editOwner(this.targetBranch.workspace, changenum, owner, opts)
+		await this.p4.editOwner(destRoboWorkspace, changenum, owner, opts)
 	}
 
 	onGlobalChange(info: ChangeInfo) {
