[NO NON-PII COMMIT MESSAGE]

diff --git a/Engine/Extras/RoboMerge/v3/src/common/perforce.ts b/Engine/Extras/RoboMerge/v3/src/common/perforce.ts
--- a/Engine/Extras/RoboMerge/v3/src/common/perforce.ts
+++ b/Engine/Extras/RoboMerge/v3/src/common/perforce.ts
@@ -1221,18 +1221,18 @@ export class PerforceContext {
 
 	// change the description on an existing CL
 	// output format is just error or not
-	async editDescription(roboWorkspace: RoboWorkspace, changelist: number, description: string) {
+	async editDescription(roboWorkspace: RoboWorkspace, changelist: number, description: string, edgeServerAddress?: string) {
 		const workspace = coercePerforceWorkspace(roboWorkspace);
 
 		// get the current changelist description
-		const output = await this._execP4(workspace, ['change', '-o', changelist.toString()]);
+		const output = await this._execP4(workspace, ['change', '-o', changelist.toString()], {edgeServerAddress});
 		// replace the description
 		let new_desc = '\nDescription:\n\t' + this._sanitizeDescription(description);
 		let form = output.replace(/\nDescription:\n(\t[^\n]*\n)*/, new_desc.replace(/\$/g, '$$$$'));
 
 		// run the P4 change command to update
 		this.logger.info("Executing: 'p4 change -i -u' to edit description on CL" + changelist);
-		await this._execP4(workspace, ['change', '-i', '-u'], { stdin: form, quiet: true });
+		await this._execP4(workspace, ['change', '-i', '-u'], { stdin: form, quiet: true, edgeServerAddress });
 	}
 
 	// change the owner of an existing CL
diff --git a/Engine/Extras/RoboMerge/v3/src/robo/edgebot.ts b/Engine/Extras/RoboMerge/v3/src/robo/edgebot.ts
--- a/Engine/Extras/RoboMerge/v3/src/robo/edgebot.ts
+++ b/Engine/Extras/RoboMerge/v3/src/robo/edgebot.ts
@@ -805,7 +805,7 @@ class EdgeBotImpl extends PerforceStatefulBot {
 		}
 		else {
 			// edit the CL description
-			await this.p4.editDescription(destRoboWorkspace, changenum, final_desc)
+			await this.p4.editDescription(destRoboWorkspace, changenum, final_desc, edgeServerAddress)
 
 			// shelve the files as we see them (should trigger a codereview email)
 			if (!await this.p4.shelve(destRoboWorkspace, changenum, edgeServerAddress)) {
