[NO NON-PII COMMIT MESSAGE]

diff --git a/Engine/Extras/RoboMerge/v3/src/new/graph.ts b/Engine/Extras/RoboMerge/v3/src/new/graph.ts
--- a/Engine/Extras/RoboMerge/v3/src/new/graph.ts
+++ b/Engine/Extras/RoboMerge/v3/src/new/graph.ts
@@ -232,7 +232,8 @@ export class Graph {
 		return null
 	}
 
-	private findAllRoutesBetweenImpl(completeRoutes: Edge[][], routeSoFar: Edge[], src: Node, seen: Set<Node>, target: Node, flags: EdgeFlag[]) {
+	private findAllRoutesBetweenImpl(completeRoutes: Edge[][], routeSoFar: Edge[],
+																	src: Node, seen: Set<Node>, target: Node, flags: EdgeFlag[]) {
 		for (const edge of this.getEdgesBySource(src, ...flags)) {
 			if (!seen.has(edge.target)) {
 				const routePlus = [...routeSoFar, edge]
@@ -281,7 +282,8 @@ export class Graph {
 			// flood the graph one step to include all unseen direct flowsTo nodes of one branch
 			let anyUnseen = false
 
-			let flowsTo: Set<Edge> | Edge[] | undefined = this.edgesBySource.get(sourceNode) // here's where we will need to be able to filter by bot
+		 // here's where we will need to be able to filter by bot
+			let flowsTo: Set<Edge> | Edge[] | undefined = this.edgesBySource.get(sourceNode)
 			if (flowsTo && allowedBots) {
 				// console.log('before: ', flowsTo)
 				flowsTo = [...flowsTo].filter((e: Edge) => allowedBots.indexOf(e.bot) >= 0)
@@ -335,11 +337,11 @@ export class Graph {
 
 export function addBranchGraph(graph: Graph, branchGraph: BranchGraphInterface) {
 
-	graph.branchGraphAliases.set(branchGraph.botname.toUpperCase(), branchGraph)
-	if (branchGraph.config.alias) {
-		graph.branchGraphAliases.set(branchGraph.config.alias.toUpperCase(), branchGraph)
-
-	}
+	graph.branchGraphAliases  = new Map([
+		[branchGraph.botname.toUpperCase(), branchGraph],
+		...graph.branchGraphAliases.entries(),
+		...(branchGraph.config.aliases.map(s => [s.toUpperCase(), branchGraph]) as [[string, BranchGraphInterface]])
+	])
 
 	const botname = branchGraph.botname as BotName
 	const branchNodes = new Map<Branch, Node>()
@@ -648,7 +650,8 @@ function parseTargetsAndFlagsImpl(tokens: string[], logger: ContextualLogger, fo
 	return {status: 'succeeded', flags, targets}
 }
 
-export function parseTargetsAndFlags(graph: Graph, botname: BotName, tokens: string[], automaticTargets: Node[], logger: ContextualLogger, forcedMode?: MergeMode) {
+export function parseTargetsAndFlags(graph: Graph, botname: BotName, tokens: string[], automaticTargets: Node[],
+																logger: ContextualLogger, forcedMode?: MergeMode) {
 	const result = parseTargetsAndFlagsImpl(tokens, logger, forcedMode)
 	if (result.status !== 'succeeded') {
 		throw new Error('to do!')
@@ -718,7 +721,8 @@ function parseChange(cldesc: string) {
 		}
 
 		if (bot) {
-			const tokensForThisLine: StringBotNamePair[] = value.split(/[ ,]/).filter(Boolean).map(token => [token, bot] as StringBotNamePair)
+			const tokensForThisLine: StringBotNamePair[] = value.split(/[ ,]/)
+													.filter(Boolean).map(token => [token, bot] as StringBotNamePair)
 			tokens = [...tokens, ...tokensForThisLine]
 		}
 	}
diff --git a/Engine/Extras/RoboMerge/v3/src/robo/branchdefs.ts b/Engine/Extras/RoboMerge/v3/src/robo/branchdefs.ts
--- a/Engine/Extras/RoboMerge/v3/src/robo/branchdefs.ts
+++ b/Engine/Extras/RoboMerge/v3/src/robo/branchdefs.ts
@@ -25,7 +25,7 @@ export interface BotConfig {
 	slackChannel: string
 	reportToBuildHealth: boolean
 	mirrorPath: string[]
-	alias: string // alias if we need to mask name of bot in commands
+	aliases: string[] // alternative names for the bot, first in list used for incognito mode
 	badgeUrlOverride: string
 	branchNamesToIgnore: string[]
 
@@ -256,24 +256,32 @@ export class BranchDefs {
 			slackChannel: '',
 			reportToBuildHealth: false,
 			mirrorPath: [],
-			alias: '',
+			aliases: [],
 			branchNamesToIgnore: [],
 			macros: {},
 			badgeUrlOverride: ''
 		}
 
-		let branchGraph
+		let branchGraphRaw: any
 		try {
-			branchGraph = jsonlint.parse(branchSpecsText) as BranchGraphDefinition
+			branchGraphRaw = jsonlint.parse(branchSpecsText)
+
+			if (!Array.isArray(branchGraphRaw.branches)) {
+				throw new Error('expected "branches" array!')
+			}
 		}
 		catch (err) {
 			outErrors.push(err)
 			return {branchGraphDef: null, config: defaultConfigForWholeBot}
 		}
 
+		if (branchGraphRaw.alias) {
+			branchGraphRaw.aliases = [branchGraphRaw.alias, ...(branchGraphRaw.aliases || [])]
+		}
+
 		// copy config values
 		for (let key of Object.keys(defaultConfigForWholeBot)) {
-			let value = (branchGraph as any)[key]
+			let value = branchGraphRaw[key]
 			if (value !== undefined) {
 				if (key === 'macros') {
 					let macrosLower: {[name:string]: string[]} | null = {}
@@ -310,6 +318,7 @@ export class BranchDefs {
 
 		const names = new Map<string, string>()
 
+		const branchGraph = branchGraphRaw as BranchGraphDefinition
 		const branchesFromJSON = branchGraph.branches
 		branchGraph.branches = []
 
diff --git a/Engine/Extras/RoboMerge/v3/src/robo/edgebot.ts b/Engine/Extras/RoboMerge/v3/src/robo/edgebot.ts
--- a/Engine/Extras/RoboMerge/v3/src/robo/edgebot.ts
+++ b/Engine/Extras/RoboMerge/v3/src/robo/edgebot.ts
@@ -428,8 +428,8 @@ class EdgeBotImpl extends PerforceStatefulBot {
 				target.mergeMode === 'null' ? '!' : '') + target.branchName)
 
 		if (flags || thisBotMergeCommands.length !== 0) {
-			const thisBotname = 
-				(this.incognitoMode && this.sourceNode.branchGraph.config.alias) ||
+			const botAliases = this.sourceNode.branchGraph.config.aliases
+			const thisBotname = this.incognitoMode && botAliases.length > 0 ? botAliases[0] :
 				this.graphBotName
 
 			description += `#ROBOMERGE[${thisBotname}]: ${thisBotMergeCommands.join(' ')}${flags}\n`
diff --git a/Engine/Extras/RoboMerge/v3/src/robo/nodebot.ts b/Engine/Extras/RoboMerge/v3/src/robo/nodebot.ts
--- a/Engine/Extras/RoboMerge/v3/src/robo/nodebot.ts
+++ b/Engine/Extras/RoboMerge/v3/src/robo/nodebot.ts
@@ -1560,7 +1560,7 @@ export class NodeBot extends PerforceStatefulBot implements NodeBotInterface {
 			isDefaultBot: this.branch.isDefaultBot,
 			graphBotName: this.graphBotName,
 			cl: change.change,
-			aliasUpper: (this.branchGraph.config.alias || '').toUpperCase(),
+			aliasesUpper: this.branchGraph.config.aliases.map(s => s.toUpperCase()),
 			macros: this.branchGraph.config.macros,
 			logger: this.nodeBotLogger
 		})
diff --git a/Engine/Extras/RoboMerge/v3/src/robo/targets.ts b/Engine/Extras/RoboMerge/v3/src/robo/targets.ts
--- a/Engine/Extras/RoboMerge/v3/src/robo/targets.ts
+++ b/Engine/Extras/RoboMerge/v3/src/robo/targets.ts
@@ -183,7 +183,7 @@ export class DescriptionParser {
 		private isDefaultBot: boolean,
 		private graphBotName: string,
 		private cl: number,
-		private aliasUpper: string,
+		private aliasesUpper: string[],
 		private macros: {[name: string]: string[]}
 
 	) {
@@ -270,7 +270,7 @@ export class DescriptionParser {
 				continue
 			}
 
-			const thisBotNames = ['', this.graphBotName, this.aliasUpper, 'ALL']
+			const thisBotNames = ['', this.graphBotName, ...this.aliasesUpper, 'ALL']
 			const commands: Token_Command[] = []
 			for (const tokenCommand of (result as Token_Command[])) {
 				const macroLines = this.macros[tokenCommand.param.toLowerCase()]
@@ -323,12 +323,12 @@ export function parseDescriptionLines(args: {
 		isDefaultBot: boolean, 
 		graphBotName: string, 
 		cl: number, 
-		aliasUpper: string,
+		aliasesUpper: string[],
 		macros: {[name: string]: string[]},
 		logger: ContextualLogger
 	}
 	) {
-	const lineParser = new DescriptionParser(args.isDefaultBot, args.graphBotName, args.cl, args.aliasUpper, args.macros)
+	const lineParser = new DescriptionParser(args.isDefaultBot, args.graphBotName, args.cl, args.aliasesUpper, args.macros)
 	lineParser.parse(args.lines)
 	return lineParser
 }
@@ -381,7 +381,7 @@ export function processOtherBotTargets(
 			info = {
 				firstHopBranches: new Set<Branch>(),
 				commands: [],
-				aliasOrName: bg.config.alias || bg.botname
+				aliasOrName: bg.config.aliases.length > 0 ? bg.config.aliases[0] : bg.botname
 			}
 			otherBotInfo.set(key, info)
 		}
@@ -774,7 +774,7 @@ export function runTests(parentLogger: ContextualLogger) {
 		isDefaultBot: true,
 		graphBotName: 'test',
 		cl: 1,
-		aliasUpper: '',
+		aliasesUpper: [],
 		macros: {'m': ['#robomerge X, Y'], 'otherbot': ['#robomerge[A] B']},
 		logger
 	}
