[NO NON-PII COMMIT MESSAGE]

diff --git a/Engine/Extras/RoboMerge/v3/config/ldap.cfg.json b/Engine/Extras/RoboMerge/v3/config/ldap.cfg.json
--- a/Engine/Extras/RoboMerge/v3/config/ldap.cfg.json
+++ b/Engine/Extras/RoboMerge/v3/config/ldap.cfg.json
@@ -10,7 +10,9 @@
 		"groupSearchBase": "dc=companyname,dc=net",
 		"groupSearchFilter": "(&(objectclass=group)(member:1.2.840.113556.1.4.1941:={{dn}}))",
 
-		"tlsOptions": {"rejectUnauthorized": false}
+		"tlsOptions": {"rejectUnauthorized": false},
+
+		"connectTimeout": 30000
 	},
 	"bot-groups": [
 		{
