[NO NON-PII COMMIT MESSAGE]

diff --git a/Engine/Extras/RoboMerge/v3/src/robo/branchdefs.ts b/Engine/Extras/RoboMerge/v3/src/robo/branchdefs.ts
--- a/Engine/Extras/RoboMerge/v3/src/robo/branchdefs.ts
+++ b/Engine/Extras/RoboMerge/v3/src/robo/branchdefs.ts
@@ -137,6 +137,8 @@ const edgeOptionFieldsPrototype = {
 
 	implicitCommands: [''],
 
+	ignoreInCycleDetection: false,
+
 	approval: {
 		description: '',
 		channelName: '',
diff --git a/Engine/Extras/RoboMerge/v3/src/robo/branchgraph.ts b/Engine/Extras/RoboMerge/v3/src/robo/branchgraph.ts
--- a/Engine/Extras/RoboMerge/v3/src/robo/branchgraph.ts
+++ b/Engine/Extras/RoboMerge/v3/src/robo/branchgraph.ts
@@ -427,18 +427,48 @@ if (botname === '__TEST__') {
 			// compute reachable
 			const reach = this._computeReachableFrom(new Set([branch]), 'flowsTo', branch)
 			reach.delete(branch.upperName) // loops are allowed but we don't want them in the flow list
-			branch.reachable = Array.from(reach).sort()
+			branch.reachable = [...reach].sort()
 
 			// compute forced downstream
 			const forced = this._computeReachableFrom(new Set([branch]), 'forceFlowTo', branch)
-			branch.forcedDownstream = Array.from(forced).sort()
-			if (forced.has(branch.upperName)) {
-				throw new Error(`Branch ${branch.name} has a forced flow loop to itself. ` + JSON.stringify(branch.forcedDownstream))
+			branch.forcedDownstream = [...forced].sort()
+		}
+
+		// new forced flow cycle detection
+		const visited = new Set<string>()
+		for (const branch of this.branches) {
+			const cyclePath = this.hasCycle(branch, [], visited)
+			if (cyclePath) {
+				throw new Error(`Branch ${branch.name} has a forced flow loop to itself: ${JSON.stringify(cyclePath)}`)
 			}
 		}
 	}
 
-	// public for use by NodeBot - maybe move to a utility
+	private hasCycle(branch: Branch, pathSoFar: string[], visited: Set<string>): string[] | null {
+		if (pathSoFar.indexOf(branch.upperName) >= 0) {
+			// we found a cycle
+			return pathSoFar
+		}
+		if (visited.has(branch.upperName)) {
+			// already looked for cycles from this node
+			return null
+		}
+		visited.add(branch.upperName)
+
+		const path = [...pathSoFar, branch.upperName]
+		for (const target of branch.forceFlowTo) {
+			const targetBranch = this.names.get(target.toUpperCase())!
+			if (!branch.edgeProperties.get(targetBranch.upperName)!.ignoreInCycleDetection) {
+				const cyclePath = this.hasCycle(targetBranch, path, visited)
+				if (cyclePath) {
+					return cyclePath
+				}
+			}
+		}
+		return null
+	}
+
+	// public for use by computeTargets
 	public _computeReachableFrom(visited: Set<Branch>, flowKey: string, branch: Branch) {
 		// start with flows to
 		const directFlow: Set<string> = (branch as any)[flowKey]
