[NO NON-PII COMMIT MESSAGE]

diff --git a/Engine/Extras/RoboMerge/v3/src/robo/preview.ts b/Engine/Extras/RoboMerge/v3/src/robo/preview.ts
--- a/Engine/Extras/RoboMerge/v3/src/robo/preview.ts
+++ b/Engine/Extras/RoboMerge/v3/src/robo/preview.ts
@@ -22,16 +22,19 @@ export async function getPreview(cl: number, singleBot?: string) {
 
 	const bots: [string, string][] = []
 	for (const entry of (await p4.describe(cl, undefined, true)).entries) {
-		const match = entry.depotFile.match(/.*\/(.*)\.branchmap\.json$/)
-		if (match) {
-			bots.push([match[1], match[0]])
+		if (entry.action !== 'delete') {
+			const match = entry.depotFile.match(/.*\/(.*)\.branchmap\.json$/)
+			if (match) {
+				bots.push([match[1], match[0]])
+			}
 		}
 	}
 
 	const allStreamSpecs = await p4.streams()
 
+	const errors = []
 	for (const [bot, path] of bots) {
-		if (singleBot && bot.toLowerCase() !== singleBot.toLowerCase()) {
+		if (singleBot && bot.toLowerCase() !== singleBot.toLowerCase()) {	
 			continue
 		}
 		const fileText = await p4.print(`${path}@=${cl}`)
@@ -39,18 +42,29 @@ export async function getPreview(cl: number, singleBot?: string) {
 		const validationErrors: string[] = []
 		const result = BranchDefs.parseAndValidate(validationErrors, fileText, allStreamSpecs)
 
+		const errorPrefix = `\n\t${bot} validation failed: `
 		if (!result.branchGraphDef) {
-			throw new Error(validationErrors.length === 0 ? 'Failed to parse' : validationErrors.join('\n'))
+			errors.push(errorPrefix + (validationErrors.length === 0 ? 'unknown error' : validationErrors.join('')))
+			continue
 		}
 
 		const graph = new BranchGraph(bot)
 		graph.config = result.config
-		graph._initFromBranchDefInternal(result.branchGraphDef)
+		try {
+			graph._initFromBranchDefInternal(result.branchGraphDef)
+		}
+		catch (err) {
+			errors.push(errorPrefix + err.toString())
+			continue
+		}
 
 		for (const branch of graph.branches) {
 			status.addBranch(branch)
 		}
 	}
 
+	if (errors) {
+		throw new Error(errors.join('\n\n'))
+	}
 	return status
 }
