[NO NON-PII COMMIT MESSAGE]

diff --git a/Engine/Extras/RoboMerge/v3/src/robo/preview.ts b/Engine/Extras/RoboMerge/v3/src/robo/preview.ts
--- a/Engine/Extras/RoboMerge/v3/src/robo/preview.ts
+++ b/Engine/Extras/RoboMerge/v3/src/robo/preview.ts
@@ -63,7 +63,7 @@ export async function getPreview(cl: number, singleBot?: string) {
 		}
 	}
 
-	if (errors) {
+	if (errors.length > 0) {
 		throw new Error(errors.join('\n\n'))
 	}
 	return status
