[NO NON-PII COMMIT MESSAGE]

diff --git a/Engine/Extras/RoboMerge/v3/src/new/graph.ts b/Engine/Extras/RoboMerge/v3/src/new/graph.ts
--- a/Engine/Extras/RoboMerge/v3/src/new/graph.ts
+++ b/Engine/Extras/RoboMerge/v3/src/new/graph.ts
@@ -249,7 +249,6 @@ export class Graph {
 		}
 	}
 
-	// function called this to match NodeBot function
 	computeImplicitTargets(source: Node, requestedTargets: Map<Node, MergeMode>, allowedBots?: string[])
 		: {status: Success, integrations?: ComputeResult, unreachable?: Node[]} {
 
@@ -282,12 +281,9 @@ export class Graph {
 			// flood the graph one step to include all unseen direct flowsTo nodes of one branch
 			let anyUnseen = false
 
-		 // here's where we will need to be able to filter by bot
 			let flowsTo: Set<Edge> | Edge[] | undefined = this.edgesBySource.get(sourceNode)
 			if (flowsTo && allowedBots) {
-				// console.log('before: ', flowsTo)
 				flowsTo = [...flowsTo].filter((e: Edge) => allowedBots.indexOf(e.bot) >= 0)
-				// console.log('after: ', flowsTo)
 			}
 			if (flowsTo && !skipNodes.has(sourceNode)) {
 				for (const edge of flowsTo) {
@@ -319,7 +315,7 @@ export class Graph {
 
 				for (let index = 1; index < explicit.length; ++index) {
 					const indirectEdge = explicit[index]
-					if (indirectTargets.indexOf(indirectEdge) === -1) {
+					if (indirectTargets.indexOf(indirectEdge) < 0) {
 						indirectTargets.push(indirectEdge)
 					}
 					targetsToFind.delete(indirectEdge.target)
@@ -333,6 +329,7 @@ export class Graph {
 
 		return {status: 'succeeded', integrations: merges}
 	}
+
 }
 
 export function addBranchGraph(graph: Graph, branchGraph: BranchGraphInterface) {
@@ -380,16 +377,16 @@ export function addBranchGraph(graph: Graph, branchGraph: BranchGraphInterface)
 		}
 	}
 
-	if (branchNodes.size > 1) {
-		//testing syntax
-		const nodeIter = branchNodes.values()
+	// if (branchNodes.size > 1) {
+	// 	//testing syntax
+	// 	const nodeIter = branchNodes.values()
 
-		const src = nodeIter.next().value
-		const targets = new Map<Node, MergeMode>()
-		targets.set(nodeIter.next().value, 'normal')
+	// 	const src = nodeIter.next().value
+	// 	const targets = new Map<Node, MergeMode>()
+	// 	targets.set(nodeIter.next().value, 'normal')
 
-		graph.computeImplicitTargets(src, targets)
-	}
+	// 	graph.computeImplicitTargets(src, targets)
+	// }
 }
 
 export type MergeMode = 'safe' | 'normal' | 'null' | 'clobber' | 'skip'
@@ -519,43 +516,67 @@ class Test {
 
 
 
+const colors = require('colors')
+colors.enable()
+colors.setTheme(require('colors/themes/generic-logging.js'))
+
+
 export function runTests(parentLogger: ContextualLogger) {
 
 
 
 	// runUnhandledRejectionTest()
 
+/**
+
+Test format:
+
+[[command, node, ...<graph definition>], expected]
+
+graph definition e.g.: ['b', 'c', ''] means a->b, b->c and d also exists
+expected: ;-separated <direct>:<indirect>
+
+4: A -> B -> c		c-b expected to produce D:C (A->D->C)
+	 => D -> c
+
+5: A => B => C		-c expected to produce B:-C
+	 -> D
+
+
+*/
 
 
 
 	const unitTestLogger = parentLogger.createChild('Graph')
 	const tests: [string[], string | null][] = [
-		[['c-b', 'a', 'b', 'c', ''], null],
-		[['c-b', 'a', 'B', 'c', ''], null],
-		[['c-d', 'a', 'bd', 'c', '', 'c'], 'B:C'],
-		[['c-b', 'a', 'bD', 'c', '', 'c'], 'D:C'],
-		[['-c', 'a', 'Bd', 'C', '', ''], 'B:-C'],
-		[['-c-b', 'a', 'B', 'C', ''], '!B:-C'],
-		[['b', 'a', 'b', 'C', ''], 'B:'],
-		[['f', 'b', 'bcd', 'a', 'a', 'Ae', 'Df', 'Eg', 'F'], 'A:DEF'],
-		[['b', 'f', 'bcd', 'a', 'a', 'Ae', 'Df', 'Eg', 'F'], 'E:DAB'],
-		[['a', 'g', 'bcd', 'a', 'a', 'Ae', 'Df', 'Eg', 'F'], 'F:EDA'],
-		[['-a', 'g', 'bcd', 'a', 'a', 'Ae', 'Df', 'Eg', 'F'], 'F:ED-A'],
-		[['de', 'a', 'b', 'c', 'de', '', ''], 'B:CED'],
-		[['de', 'h', 'hc', 'hd', 'deFg', 'hbc', 'c', 'c', 'c', 'abd'], 'D:CE'],
-		[['de', 'h', 'bd', 'hc', 'hd', 'eFg', 'hbc', 'c', 'c', 'c'], 'C:DE'],
-		[['db', 'h', 'abd', 'hc', 'hd', 'deFg', 'hbc', 'c', 'c', 'c'], 'C:DEB'],
-		[['cg', 'a', 'be', 'acd', 'b', 'b', 'aFg', 'e', 'e'], 'B:C;E:G'],
-		[['cf', 'a', 'bE', 'acd', 'b', 'b', 'aFg', 'e', 'e'], 'B:C;E:F'],
-		[['c-e', 'a', 'bE', 'acd', 'b', 'b', 'aFg', 'e', 'e'], 'B:C'],
-		[['cf-b', 'a', 'bE', 'acd', 'b', 'b', 'aFg', 'e', 'e'], null],
-		[['fg', 'h', 'abd', 'hC', 'hd', 'defg', 'hbc', 'c', 'c', 'c'], 'C:DGF'],
-		[['dge', 'h', 'abd', 'hC', 'hd', 'defg', 'hbc', 'c', 'c', 'c'], 'C:DGE'],
-		[['e', 'i', 'abcd', 'Ie', 'if', 'iG', 'ih', 'a', 'b', 'c', 'D'], 'D:GBE']
+
+		[['-c', 'a', 'b', 'c'], ''],
+		[['c-b', 'a', 		'b', 'c', ''],										null],
+		[['c-b', 'a', 		'B', 'c', ''],										null],
+		[['c-d', 'a', 		'bd', 'c', '', 'c'],								'B:C'],
+		[['c-b', 'a', 		'bD', 'c', '', 'c'],								'D:C'], //4
+		[['-c', 'a',		'Bd', 'C', '', ''],									'B:-C'],
+		[['-c-b', 'a',		'B', 'C', ''],										'!B:-C'],
+		[['b', 'a', 		'b', 'C', ''],										'B:'],
+		[['f', 'b',			'bcd', 'a', 'a', 'Ae', 'Df', 'Eg', 'F'],			'A:DEF'],
+		[['b', 'f',			'bcd', 'a', 'a', 'Ae', 'Df', 'Eg', 'F'],			'E:DAB'],
+		[['a', 'g',			'bcd', 'a', 'a', 'Ae', 'Df', 'Eg', 'F'],			'F:EDA'],
+		[['-a', 'g',		'bcd', 'a', 'a', 'Ae', 'Df', 'Eg', 'F'],			'F:ED-A'],
+		[['de', 'a',		'b', 'c', 'de', '', ''],							'B:CED'],
+		[['de', 'h',		'hc', 'hd', 'deFg', 'hbc', 'c', 'c', 'c', 'abd'],	'D:CE'],
+		[['de', 'h',		'bd', 'hc', 'hd', 'eFg', 'hbc', 'c', 'c', 'c'],		'C:DE'],
+		[['db', 'h',		'abd', 'hc', 'hd', 'deFg', 'hbc', 'c', 'c', 'c'],	'C:DEB'],
+		[['cg', 'a',		'be', 'acd', 'b', 'b', 'aFg', 'e', 'e'],			'B:C;E:G'],
+		[['cf', 'a', 		'bE', 'acd', 'b', 'b', 'aFg', 'e', 'e'],			'B:C;E:F'],
+		[['c-e', 'a', 		'bE', 'acd', 'b', 'b', 'aFg', 'e', 'e'],			'B:C'],
+		[['cf-b', 'a', 		'bE', 'acd', 'b', 'b', 'aFg', 'e', 'e'],			null],
+		[['fg', 'h',		'abd', 'hC', 'hd', 'defg', 'hbc', 'c', 'c', 'c'],	'C:DGF'],
+		[['dge', 'h',		'abd', 'hC', 'hd', 'defg', 'hbc', 'c', 'c', 'c'],	'C:DGE'],
+		[['e', 'i',			'abcd', 'Ie', 'if', 'iG', 'ih', 'a', 'b', 'c', 'D'],'D:GBE']
 	]
 
-	let ran = 0, success = 0, fail = 0
-	for (const [testStr, expected] of tests) {
+	let ran = 0, success = 0, fail = 0, index = 0
+	for (const [testStr, expected] of tests /*.slice(3, 4)*/) {
 		const test = new Test(testStr.slice(2))
 		const result = test.computeTargets(testStr[1], testStr[0])
 
@@ -565,17 +586,13 @@ export function runTests(parentLogger: ContextualLogger) {
 		const succeeded = result.status === 'succeeded'
 
 		const formattedResult = succeeded ? test.formatTestComputeTargetsResult(result.integrations!) : 'failed'
-		if (expected) {
+		if (expected !== null) {
 			if (!succeeded) {
 				expectedOnFail = expected
-				++fail
 			}
 			else {
-				if (formattedResult.toUpperCase().replace(/\s|,/g, '') === expected)
-					++success
-				else {
+				if (formattedResult.toUpperCase().replace(/\s|,/g, '') !== expected) {
 					expectedOnFail = expected
-					++fail
 				}
 			}
 		}
@@ -583,10 +600,17 @@ export function runTests(parentLogger: ContextualLogger) {
 			expectedOnFail = 'fail'
 		}
 		if (expectedOnFail) {
-			unitTestLogger.warn(`MISMATCH!!!    ${formattedResult} vs ${expectedOnFail}`)
+			++fail
+			unitTestLogger.warn(`Test ${colors.warn(index.toString().padStart(2))} failed:   ` +
+				`${colors.warn(formattedResult.padStart(10))} vs ${colors.warn(expectedOnFail)}`)
+		}
+		else {
+			++success
 		}
+		++index
 	}
-	unitTestLogger.info(`Graph tests: ran ${ran} tests, ${success} matched, ${fail} failed`)
+	const message = `Graph tests: ran ${ran} tests, ${success} matched, ${fail} failed`
+	unitTestLogger.info((fail > 0 ? colors.error : colors.info)(message))
 	return fail
 }
 
