[NO NON-PII COMMIT MESSAGE]

diff --git a/Engine/Extras/RoboMerge/v3/functional_tests/framework/src/tests.ts b/Engine/Extras/RoboMerge/v3/functional_tests/framework/src/tests.ts
--- a/Engine/Extras/RoboMerge/v3/functional_tests/framework/src/tests.ts
+++ b/Engine/Extras/RoboMerge/v3/functional_tests/framework/src/tests.ts
@@ -10,6 +10,7 @@ import { ConfirmTextResolve } from './tests/confirm-text-resolve'
 import { ConfirmTextResolveBinaryStomp } from './tests/confirm-text-resolve-binary-stomp'
 import { CrossDepotStreamIntegration } from './tests/cross-depot-stream-integration'
 import { EdgeIndependence } from './tests/edge-independence'
+import { EdgeInitialCl } from './tests/edge-initial-cl'
 import { ExclusiveCheckout } from './tests/exclusive-checkout'
 import { ExcludeAuthors } from './tests/exclude-authors'
 import { ExcludeAuthorsPerEdge } from './tests/exclude-authors-per-edge'
@@ -22,6 +23,7 @@ import { IndirectTarget } from './tests/indirect-target'
 import { MergeMainRevToMultipleRelease } from './tests/merge-main-rev-to-multiple-release'
 import { MergeMainRevToRelease } from './tests/merge-main-rev-to-release'
 import { MultipleConflicts } from './tests/multiple-conflicts'
+import { MultipleRoutesToSkip } from './tests/multiple-routes-to-skip'
 import { OverriddenCommand } from './tests/overridden-command'
 import { RejectBranchResolveStomp } from './tests/reject-branch-resolve-stomp'
 import { RejectDeleteResolveStomp } from './tests/reject-delete-resolve-stomp'
@@ -37,7 +39,6 @@ import { TestChain } from './tests/test-chain'
 import { TestEdgeGate } from './tests/test-edge-gate'
 import { TestFlags } from './tests/test-flags'
 import { TestGate } from './tests/test-gate'
-import { EdgeInitialCl } from './tests/edge-initial-cl'
 import { TestReconsider } from './tests/test-reconsider'
 import { TestEdgeReconsider } from './tests/test-edge-reconsider'
 import { TestTerminal } from './tests/test-terminal'
@@ -85,7 +86,8 @@ async function addToRoboMerge(p4: Perforce, tests: FunctionalTest[]) {
 	}
 
 	await Promise.all(settings.map(([botName, s]) => 
-		P4Util.addFile(rootClient, botName + '.branchmap.json', JSON.stringify({...DEFAULT_BOT_SETTINGS, ...s, slackChannel: botName, alias: botName + '-alias'}))
+		P4Util.addFile(rootClient, botName + '.branchmap.json', JSON.stringify({...DEFAULT_BOT_SETTINGS, ...s, slackChannel: botName,
+			aliases: [botName + '-alias', botName + '-alias2']}))
 	))
 
 	await rootClient.submit('Adding branchspecs')
@@ -195,7 +197,8 @@ async function go() {
 
 		new StompWithAdd(p4),
 		new Ignore(p4),
-		new StompForwardingCommands(p4)
+		new StompForwardingCommands(p4),
+		new MultipleRoutesToSkip(p4)
 	]
 
 	// const testToDebug = availableTests[30]
diff --git a/Engine/Extras/RoboMerge/v3/functional_tests/framework/src/tests/multiple-routes-to-skip.ts b/Engine/Extras/RoboMerge/v3/functional_tests/framework/src/tests/multiple-routes-to-skip.ts
new file mode 100644
--- /dev/null
+++ b/Engine/Extras/RoboMerge/v3/functional_tests/framework/src/tests/multiple-routes-to-skip.ts
@@ -0,0 +1,55 @@
+import { FunctionalTest, P4Client, P4Util, Stream } from '../framework'
+
+const streams: Stream[] = [
+	{name: 'Main', streamType: 'mainline'},
+	{name: 'Dev-Child', streamType: 'development', parent: 'Main'},
+	{name: 'Dev-Grandchild', streamType: 'development', parent: 'Dev-Child'},
+	{name: 'Dev-OtherChild', streamType: 'development', parent: 'Main'},
+]
+
+export class MultipleRoutesToSkip extends FunctionalTest {
+	async setup() {
+		await this.p4.depot('stream', this.depotSpec())
+		await this.createStreamsAndWorkspaces(streams)
+
+		this.mainClient = this.getClient('Main')
+		await P4Util.addFile(this.mainClient, 'test.txt', 'Initial content')
+
+		await P4Util.submit(this.mainClient, 'Initial file')
+
+		// popluateStreams does all in parallel, ignoring dependencies
+		await this.populateStreams(streams.filter(s => s.name !== 'Dev-Grandchild'))
+		await this.populateStreams(streams.slice(2, 3)) // Dev-Grandchild
+		// await this.getClient('Release').sync()
+	}
+
+	run() {
+		// skip grand child (indirect target via both childrent)
+		return P4Util.editFileAndSubmit(this.mainClient, 'test.txt', 'change',
+													'-' + this.fullBranchName('Dev-Grandchild'))
+	}
+
+	verifyCommon() {
+		return Promise.all([
+			this.checkHeadRevision('Main', 'test.txt', 2),
+			this.checkHeadRevision('Dev-Child', 'test.txt', 2), 
+			this.checkHeadRevision('Dev-Grandchild', 'test.txt', 1),
+			this.checkHeadRevision('Dev-OtherChild', 'test.txt', 2)
+			])
+	}
+
+	verify() {
+		return this.verifyCommon()
+	}
+
+	getBranches() {
+		return [
+			this.makeForceAllBranchDef('Main', ['Dev-Child', 'Dev-OtherChild']),
+			this.makeForceAllBranchDef('Dev-Child', ['Dev-Grandchild']),
+			this.makeForceAllBranchDef('Dev-Grandchild', []),
+			this.makeForceAllBranchDef('Dev-OtherChild', ['Dev-Grandchild'])
+		]
+	}
+
+	mainClient: P4Client
+}
diff --git a/Engine/Extras/RoboMerge/v3/src/robo/targets.ts b/Engine/Extras/RoboMerge/v3/src/robo/targets.ts
--- a/Engine/Extras/RoboMerge/v3/src/robo/targets.ts
+++ b/Engine/Extras/RoboMerge/v3/src/robo/targets.ts
@@ -652,8 +652,16 @@ function computeTargetsImpl(
 		return {computeResult: null, errors: []}
 	}
 
+	// handle multiple routes to skip targets
+	//	
+	const skipTargets: Branch[] = [...allIntegrations]
+		.filter(([_, mergeMode]) => mergeMode === 'skip')
+		.map(([branchName, _]) => branchGraph.getBranch(branchName)!)
+
+
 	const merges = new Map<Branch, Branch[]>()
 	for (const [initialEdge, furtherEdges] of integrations) {
+
 		const edgeTargetBranch = (e: Edge) => {
 			const branch = branchGraph.getBranch(e.targetName)
 			if (!branch) {
@@ -662,7 +670,28 @@ function computeTargetsImpl(
 			return branch
 		}
 
-		merges.set(edgeTargetBranch(initialEdge), furtherEdges.map(edgeTargetBranch))
+		const targetBranch = edgeTargetBranch(initialEdge)
+		const furtherEdgeBranches = furtherEdges.map(edgeTargetBranch)
+		for (const skipTarget of skipTargets) {
+			console.log('skip stuff!!!', initialEdge.source.debugName, skipTarget.name, furtherEdgeBranches)
+			console.log('skip stuff!!!', targetBranch.forcedDownstream)
+			console.log('skip stuff!!!', furtherEdgeBranches.map(b => b.forcedDownstream))
+			if (furtherEdgeBranches.indexOf(skipTarget) < 0
+				&&
+
+			// make sure skipTarget is present directly if reachable
+				[...furtherEdgeBranches, targetBranch]
+					.map(b => b.forcedDownstream)
+					.filter(ds => (ds || []).map(b => branchGraph.getBranch(b)).indexOf(skipTarget) >= 0)
+					.length > 0
+				) {
+
+				logger.info(`adding skip of ${skipTarget.name} to onward integrations for ` +
+					`${initialEdge.source.debugName}->${initialEdge.target.debugName} due to ... todo`)
+				furtherEdgeBranches.push(skipTarget)
+			}
+		}
+		merges.set(targetBranch, furtherEdgeBranches)
 	}
 
 	return {computeResult: {merges, targets, flags: ri.flags}, errors}
