[NO NON-PII COMMIT MESSAGE]

diff --git a/Engine/Extras/RoboMerge/v3/public/js/main.js b/Engine/Extras/RoboMerge/v3/public/js/main.js
--- a/Engine/Extras/RoboMerge/v3/public/js/main.js
+++ b/Engine/Extras/RoboMerge/v3/public/js/main.js
@@ -100,7 +100,7 @@ function renderBranchList(data) {
 					$('<div>')
 					.text(botState.lastError.nodeBot + ' failed:')
 					.append($('<pre>').text(botState.lastError.error))
-					.append($restartButton)
+					// .append($restartButton)
 				);
 			}
 			else {
