[NO NON-PII COMMIT MESSAGE]

diff --git a/Engine/Extras/RoboMerge/v3/public/js/main.js b/Engine/Extras/RoboMerge/v3/public/js/main.js
--- a/Engine/Extras/RoboMerge/v3/public/js/main.js
+++ b/Engine/Extras/RoboMerge/v3/public/js/main.js
@@ -100,7 +100,7 @@ function renderBranchList(data) {
 					$('<div>')
 					.text(botState.lastError.nodeBot + ' failed:')
 					.append($('<pre>').text(botState.lastError.error))
-					// .append($restartButton)
+					.append($restartButton)
 				);
 			}
 			else {
diff --git a/Engine/Extras/RoboMerge/v3/src/robo/graph-interface.ts b/Engine/Extras/RoboMerge/v3/src/robo/graph-interface.ts
--- a/Engine/Extras/RoboMerge/v3/src/robo/graph-interface.ts
+++ b/Engine/Extras/RoboMerge/v3/src/robo/graph-interface.ts
@@ -1,6 +1,7 @@
 // Copyright Epic Games, Inc. All Rights Reserved.
 
 import { BranchGraphInterface } from './branch-interfaces';
+import { NodeBotInterface } from './bot-interfaces';
 import { BotConfig, BranchGraphDefinition } from './branchdefs';
 
 export interface GraphInterface {
@@ -9,5 +10,5 @@ export interface GraphInterface {
 
 	ensureStopping(): boolean
 	reinitFromBranchGraphsObject(config: BotConfig, branchGraphs: BranchGraphDefinition): Promise<void>
-	handleRequestedIntegrationsForAllNodes(): Promise<void>
+	handleRequestedIntegrationsForAllNodes(): Promise<[NodeBotInterface, Error] | null>
 }
diff --git a/Engine/Extras/RoboMerge/v3/src/robo/graphbot.ts b/Engine/Extras/RoboMerge/v3/src/robo/graphbot.ts
--- a/Engine/Extras/RoboMerge/v3/src/robo/graphbot.ts
+++ b/Engine/Extras/RoboMerge/v3/src/robo/graphbot.ts
@@ -130,7 +130,14 @@ export class GraphBot implements GraphInterface, BotEventHandler {
 			if (branch.enabled) {
 				const persistence = this.settings.getContext(branch.upperName)
 				branch.bot = new NodeBot(branch, this.mailer, this.externalUrl, this.eventTriggers, persistence, ubergraph,
-					() => this.handleRequestedIntegrationsForAllNodes()
+					async () => {
+						const errPair = await this.handleRequestedIntegrationsForAllNodes()
+						if (errPair) {
+							// can report wrong nodebot (this rather than one that errored)
+							const [_, err] = errPair
+							throw err
+						}
+					}
 				)
 
 				if (branch.bot.getNumConflicts() > 0) {
@@ -198,7 +205,7 @@ export class GraphBot implements GraphInterface, BotEventHandler {
 		this.crashRequested = msg
 	}
 
-	async handleRequestedIntegrationsForAllNodes() {
+	async handleRequestedIntegrationsForAllNodes(): Promise<[NodeBot, Error] | null> {
 		for (const branchName of this.branchGraph.getBranchNames()) {
 			const node = this.findNode(branchName)!
 			for (;;) {
@@ -210,12 +217,12 @@ export class GraphBot implements GraphInterface, BotEventHandler {
 					await node.processQueuedChange(request)
 				}
 				catch(err) {
-					this.handleNodebotError(node, err)
-					return
+					return [node, err]
 				}
 				node.persistQueuedChanges()
 			}
 		}
+		return null
 	}
 
 	// mostly for nodebots, but could also be the auto reloader bot
@@ -302,7 +309,12 @@ export class GraphBot implements GraphInterface, BotEventHandler {
 				await new Promise(done => setTimeout(done, this.waitTime!))
 			}
 
-			await this.handleRequestedIntegrationsForAllNodes()
+			const errPair = await this.handleRequestedIntegrationsForAllNodes()
+			if (errPair) {
+				const [bot, err] = errPair
+				this.handleNodebotError(bot, err)
+				return
+			}
 
 			roboAnalytics!.reportActivity(this.branchGraph.botname, activity)
 			roboAnalytics!.reportMemoryUsage('main', process.memoryUsage().heapUsed)
