[NO NON-PII COMMIT MESSAGE]

diff --git a/Engine/Extras/RoboMerge/v3/package.json b/Engine/Extras/RoboMerge/v3/package.json
--- a/Engine/Extras/RoboMerge/v3/package.json
+++ b/Engine/Extras/RoboMerge/v3/package.json
@@ -11,7 +11,7 @@
   "scripts": {
     "start": "node dist/robo/watchdog.js -botname=TEST -devMode",
     "test": "node dist/robo/unit-tests.js",
-    "postinstall": "./install-graphviz.js"
+    "postinstall": "node install-graphviz.js"
   },
   "dependencies": {
     "@hpcc-js/graph": "^2.27.0",
@@ -25,10 +25,11 @@
     "@types/sanitize-html": "^1.20.2",
     "@types/semver": "^6.0.2",
     "@types/swagger-ui-dist": "^3.0.3",
+    "colors": "^1.4.0",
     "jsonlint": "^1.6.3",
     "ldapauth-fork": "^4.0.2",
     "markdown-toc": "^1.2.0",
-    "marked": "^0.7.0",
+    "marked": "^4.0.12",
     "mustache": "^3.2.1",
     "node-gzip": "^1.1.2",
     "nodemailer": "^6.7.0",
