[NO NON-PII COMMIT MESSAGE]

diff --git a/Engine/Extras/RoboMerge/v3/src/robo/branch-interfaces.ts b/Engine/Extras/RoboMerge/v3/src/robo/branch-interfaces.ts
--- a/Engine/Extras/RoboMerge/v3/src/robo/branch-interfaces.ts
+++ b/Engine/Extras/RoboMerge/v3/src/robo/branch-interfaces.ts
@@ -151,7 +151,6 @@ export interface ChangeInfo extends TargetInfo {
 	targetWorkspaceOverride?: string
 	overriddenCommand: string
 	macros: string[]
-	hasOkForGithubTag: boolean
 }
 
 export interface PendingChange {
diff --git a/Engine/Extras/RoboMerge/v3/src/robo/branchdefs.ts b/Engine/Extras/RoboMerge/v3/src/robo/branchdefs.ts
--- a/Engine/Extras/RoboMerge/v3/src/robo/branchdefs.ts
+++ b/Engine/Extras/RoboMerge/v3/src/robo/branchdefs.ts
@@ -133,7 +133,6 @@ const edgeOptionFieldsPrototype = {
 	additionalSlackChannel: '',
 
 	terminal: false, // changes go along terminal edges but no further
-	doHackyOkForGithubThing: false,
 
 	implicitCommands: [''],
 
diff --git a/Engine/Extras/RoboMerge/v3/src/robo/edgebot.ts b/Engine/Extras/RoboMerge/v3/src/robo/edgebot.ts
--- a/Engine/Extras/RoboMerge/v3/src/robo/edgebot.ts
+++ b/Engine/Extras/RoboMerge/v3/src/robo/edgebot.ts
@@ -171,9 +171,6 @@ class EdgeBotImpl extends PerforceStatefulBot {
 	get disallowSkip() {
 		return !!this.options.disallowSkip
 	}
-	get doHackyOkForGithubThing() {
-		return !!this.options.doHackyOkForGithubThing
-	}
 	get implicitCommands() {
 		return this.options.implicitCommands || []
 	}
@@ -369,14 +366,6 @@ class EdgeBotImpl extends PerforceStatefulBot {
 			description += `#ROBOMERGE-AUTHOR: ${authorTag}\n`
 		}
 
-		if (this.doHackyOkForGithubThing/* && !info.hasOkForGithubTag */) {
-			const descriptionLines = description.split('\n')
-			description = descriptionLines
-				.filter(s => !s.toLowerCase().startsWith('#okforgithub'))
-				.join('\n')
-			description += '#okforgithub ignore\n'
-		}
-
 		if (this.isTerminal) {
 			description += '#robomerge #ignore\n'
 		}
diff --git a/Engine/Extras/RoboMerge/v3/src/robo/nodebot.ts b/Engine/Extras/RoboMerge/v3/src/robo/nodebot.ts
--- a/Engine/Extras/RoboMerge/v3/src/robo/nodebot.ts
+++ b/Engine/Extras/RoboMerge/v3/src/robo/nodebot.ts
@@ -1609,7 +1609,6 @@ export class NodeBot extends PerforceStatefulBot implements NodeBotInterface {
 			forceCreateAShelf, sendNoShelfEmail,
 
 			forceStompChanges, additionalDescriptionText,
-			hasOkForGithubTag: parsedLines.hasOkForGithubTag,
 			overriddenCommand: commandOverride,
 			macros: parsedLines.expandedMacros
 		}
diff --git a/Engine/Extras/RoboMerge/v3/src/robo/targets.ts b/Engine/Extras/RoboMerge/v3/src/robo/targets.ts
--- a/Engine/Extras/RoboMerge/v3/src/robo/targets.ts
+++ b/Engine/Extras/RoboMerge/v3/src/robo/targets.ts
@@ -166,7 +166,6 @@ export class DescriptionParser {
 	descFinal: string[] = []
 
 	propagatingNullMerge = false
-	hasOkForGithubTag = false
 	useDefaultFlow = true
 
 	// arguments: string[] = []
@@ -241,9 +240,6 @@ export class DescriptionParser {
 			return false
 		}
 
-		if (token.tag === 'OKFORGITHUB') {
-			this.hasOkForGithubTag = true
-		}
 		return true
 	}
 
