[NO NON-PII COMMIT MESSAGE]

diff --git a/Engine/Extras/RoboMerge/v3/src/robo/edgebot.ts b/Engine/Extras/RoboMerge/v3/src/robo/edgebot.ts
--- a/Engine/Extras/RoboMerge/v3/src/robo/edgebot.ts
+++ b/Engine/Extras/RoboMerge/v3/src/robo/edgebot.ts
@@ -557,7 +557,7 @@ class EdgeBotImpl extends PerforceStatefulBot {
 
 		// Revert attempt
 		if (pending.newCl > 0) {
-			await this.revertAndDelete(pending.newCl)
+			await this.revertAndDelete(pending.newCl, edgeServerAddress)
 			pending.newCl = -1
 		}
 		
@@ -574,12 +574,12 @@ class EdgeBotImpl extends PerforceStatefulBot {
 		return new EdgeIntegrationDetails('error', errors.join('\n'))
 	}
 
-	private async revertAndDelete(cl: number) {
+	private async revertAndDelete(cl: number, edgeServerAddress?: string) {
 		this._log_action(`Reverting CL ${cl}`);
-		await this.p4.revert(this.targetBranch.workspace, cl);
+		await this.p4.revert(this.targetBranch.workspace, cl, [], edgeServerAddress);
 
 		this._log_action(`Deleting CL ${cl}`);
-		await this.p4.deleteCl(this.targetBranch.workspace, cl);
+		await this.p4.deleteCl(this.targetBranch.workspace, cl, edgeServerAddress);
 	}
 
 	async revertPendingCLWithShelf(client: string, change: number, userContext: string) {
