[NO NON-PII COMMIT MESSAGE]

diff --git a/Engine/Extras/RoboMerge/v3/package.json b/Engine/Extras/RoboMerge/v3/package.json
--- a/Engine/Extras/RoboMerge/v3/package.json
+++ b/Engine/Extras/RoboMerge/v3/package.json
@@ -18,7 +18,7 @@
     "@hpcc-js/wasm": "^0.3.14",
     "@sentry/node": "^5.7.1",
     "@types/ldapjs": "^1.0.3",
-    "@types/marked": "^0.6.5",
+    "@types/marked": "^4.0.2",
     "@types/mustache": "^0.8.32",
     "@types/node": "^12.12.24",
     "@types/nodemailer": "^4.3.4",
diff --git a/Engine/Extras/RoboMerge/v3/src/robo/roboserver.ts b/Engine/Extras/RoboMerge/v3/src/robo/roboserver.ts
--- a/Engine/Extras/RoboMerge/v3/src/robo/roboserver.ts
+++ b/Engine/Extras/RoboMerge/v3/src/robo/roboserver.ts
@@ -1,7 +1,7 @@
 // Copyright Epic Games, Inc. All Rights Reserved.
 import * as Sentry from '@sentry/node';
 import * as fs from 'fs';
-import * as marked from 'marked';
+import {marked} from 'marked';
 import * as querystring from 'querystring';
 import { Arg, readProcessArgs } from '../common/args';
 import { ContextualLogger } from '../common/logger';
@@ -241,7 +241,7 @@ class RoboWebApp implements AppInterface {
 			const generatedTOC = toc(markdown).content
 			markdown = markdown.replace(/<!--\s*toc\s*-->/ig, generatedTOC)
 
-			marked(markdown, {renderer: RoboWebApp.getMarkdownRenderer()}, (error, parseResult) => {
+			marked.parse(markdown, {renderer: RoboWebApp.getMarkdownRenderer()}, (error, parseResult) => {
 				if (error) {
 					fail(error)
 					return
