[NO NON-PII COMMIT MESSAGE]

diff --git a/Engine/Extras/RoboMerge/v3/src/robo/session.ts b/Engine/Extras/RoboMerge/v3/src/robo/session.ts
--- a/Engine/Extras/RoboMerge/v3/src/robo/session.ts
+++ b/Engine/Extras/RoboMerge/v3/src/robo/session.ts
@@ -58,7 +58,7 @@ export class Session {
 
 	public static onLoginAttempt: ((result: string) => void) | null = null
 
-	static init(logger: ContextualLogger) {
+	static init(_logger: ContextualLogger) {
 		const ldapConfig = JSON.parse(fs.readFileSync('config/ldap.cfg.json', 'utf8'))
 
 		Session.LDAP_CONFIG = ldapConfig['server-config']
@@ -96,7 +96,6 @@ export class Session {
 
 			logger.info(`LDAP: auth helper created (${creds.user} log-in)`)
 
-
 			const startTime = Date.now()
 			auth.on('error', fail)
 			auth.authenticate(creds.user, creds.password, (err: any | null, userData: any) => {
@@ -106,6 +105,8 @@ export class Session {
 				if (err) {
 					if (err.name === 'InvalidCredentialsError' ||
 						(typeof(err) === 'string' && err.startsWith('no such user'))) {
+
+						logger.info(`Log-in failed: ${err.name}, ${err} (${creds.user})`)
 						if (Session.onLoginAttempt) {
 							Session.onLoginAttempt('fail')
 						}
