[NO NON-PII COMMIT MESSAGE]

diff --git a/Engine/Extras/RoboMerge/v3/src/robo/session.ts b/Engine/Extras/RoboMerge/v3/src/robo/session.ts
--- a/Engine/Extras/RoboMerge/v3/src/robo/session.ts
+++ b/Engine/Extras/RoboMerge/v3/src/robo/session.ts
@@ -58,7 +58,7 @@ export class Session {
 
 	public static onLoginAttempt: ((result: string) => void) | null = null
 
-	static init(_logger: ContextualLogger) {
+	static init(logger: ContextualLogger) {
 		const ldapConfig = JSON.parse(fs.readFileSync('config/ldap.cfg.json', 'utf8'))
 
 		Session.LDAP_CONFIG = ldapConfig['server-config']
