[NO NON-PII COMMIT MESSAGE]

diff --git a/Engine/Extras/RoboMerge/v3/src/common/perforce.ts b/Engine/Extras/RoboMerge/v3/src/common/perforce.ts
--- a/Engine/Extras/RoboMerge/v3/src/common/perforce.ts
+++ b/Engine/Extras/RoboMerge/v3/src/common/perforce.ts
@@ -1334,7 +1334,7 @@ export class PerforceContext {
 			args = ['-p', opts.edgeServerAddress, ...args]
 		}
 
-		args = ['-zprog=robomerge', '-zversion=' + robomergeVersion, ...args]
+		args = ['-zprog=robomerge', '-zversion=' + robomergeVersion, '-zmaxscanrows=200000000', ...args]
 
 		// log what we're running
 		let cmd_rec = new CommandRecord('p4 ' + args.join(' '));
