[NO NON-PII COMMIT MESSAGE]

diff --git a/Engine/Extras/RoboMerge/v3/src/common/perforce.ts b/Engine/Extras/RoboMerge/v3/src/common/perforce.ts
--- a/Engine/Extras/RoboMerge/v3/src/common/perforce.ts
+++ b/Engine/Extras/RoboMerge/v3/src/common/perforce.ts
@@ -227,6 +227,7 @@ export interface Change {
 	forceCreateAShelf?: boolean
 	sendNoShelfEmail?: boolean // Used for requesting stomps for internal RM usage, such as stomp changes
 	commandOverride?: string
+	accumulateCommandOverride?: boolean
 
 	// Stomp Changes support
 	forceStompChanges?: boolean
diff --git a/Engine/Extras/RoboMerge/v3/src/robo/nodebot.ts b/Engine/Extras/RoboMerge/v3/src/robo/nodebot.ts
--- a/Engine/Extras/RoboMerge/v3/src/robo/nodebot.ts
+++ b/Engine/Extras/RoboMerge/v3/src/robo/nodebot.ts
@@ -415,7 +415,18 @@ export class NodeBot extends PerforceStatefulBot implements NodeBotInterface {
 			change.forceStompChanges = true
 		}
 		change.additionalDescriptionText = fromQueue.description
-		change.commandOverride = fromQueue.commandOverride
+		if (fromQueue.commandOverride) {
+			change.commandOverride = fromQueue.commandOverride
+		}
+		// regular reconsiders, make sure edge target is included
+		else if (specifiedTargetBranch && !change.forceCreateAShelf && !change.forceStompChanges) {
+			// new behaviour: make target explicit
+			const bg = this.branchGraph
+			
+			const botNameOrAlias = bg.config.aliases.length > 0 ? bg.config.aliases[0] : bg.botname
+			change.commandOverride = `#robomerge[${botNameOrAlias}] ${specifiedTargetBranch.upperName}`
+			change.accumulateCommandOverride = true
+		}
 
 		// process just this change
 		await this._processAndMergeCl(new Map<string, EdgeBot>(this.edges), change, true, fromQueue.who, fromQueue.workspace, specifiedTargetBranch)
@@ -1569,7 +1580,14 @@ export class NodeBot extends PerforceStatefulBot implements NodeBotInterface {
 
 		const commandOverride = (change.commandOverride || '').trim()
 		if (commandOverride) {
-			parsedLines.override(parse(commandOverride.split('|').map(s => s.trim())))
+
+			const parsedOverride = parse(commandOverride.split('|').map(s => s.trim()))
+			if (change.accumulateCommandOverride) {
+				parsedLines.append(parsedOverride)
+			}
+			else {
+				parsedLines.override(parsedOverride)
+			}
 		}
 
 		// Author is always CL user
diff --git a/Engine/Extras/RoboMerge/v3/src/robo/targets.ts b/Engine/Extras/RoboMerge/v3/src/robo/targets.ts
--- a/Engine/Extras/RoboMerge/v3/src/robo/targets.ts
+++ b/Engine/Extras/RoboMerge/v3/src/robo/targets.ts
@@ -199,6 +199,15 @@ export class DescriptionParser {
 		this.errors = other.errors
 	}
 
+	append(other: DescriptionParser) {
+		// ignore default flow here
+		this.arguments = [...this.arguments, ...other.arguments]
+		this.expandedMacros = [...this.expandedMacros, ...other.expandedMacros]
+		this.otherBotArguments = [...this.otherBotArguments, ...other.otherBotArguments]
+
+		this.errors = [...this.errors, ...other.errors]
+	}
+
 	private processNonRoboToken(token: Token_Other) {
 		if (token.isRobo) {
 			if (ROBO_TAGS.indexOf(token.tag) >= 0) {
