[NO NON-PII COMMIT MESSAGE]

diff --git a/Engine/Extras/RoboMerge/v3/src/robo/branchdefs.ts b/Engine/Extras/RoboMerge/v3/src/robo/branchdefs.ts
--- a/Engine/Extras/RoboMerge/v3/src/robo/branchdefs.ts
+++ b/Engine/Extras/RoboMerge/v3/src/robo/branchdefs.ts
@@ -130,6 +130,7 @@ type NodeOptionFields = typeof nodeOptionFieldsPrototype
 const edgeOptionFieldsPrototype = {
 	...commonOptionFieldsPrototype,
 
+	branchspec: '',
 	additionalSlackChannel: '',
 
 	terminal: false, // changes go along terminal edges but no further
diff --git a/Engine/Extras/RoboMerge/v3/src/robo/branchgraph.ts b/Engine/Extras/RoboMerge/v3/src/robo/branchgraph.ts
--- a/Engine/Extras/RoboMerge/v3/src/robo/branchgraph.ts
+++ b/Engine/Extras/RoboMerge/v3/src/robo/branchgraph.ts
@@ -149,7 +149,7 @@ if (botname === '__TEST__') {
 			}
 		}
 
-		// set up branchspecs to use
+		// old style branchspecs
 		if (branchGraph.branchspecs) {
 			for (let def of branchGraph.branchspecs) {
 				// get the name
@@ -169,11 +169,11 @@ if (botname === '__TEST__') {
 
 				// add forward and reverse entries
 				if (!fromBranch.config.ignoreBranchspecs) {
-					fromBranch.branchspec.set(toBranch.upperName, {name: name, reverse: false})
+					fromBranch.branchspec.set(toBranch.upperName, {name, reverse: false})
 				}
 
 				if (!toBranch.config.ignoreBranchspecs) {
-					toBranch.branchspec.set(fromBranch.upperName, {name: name, reverse: true})
+					toBranch.branchspec.set(fromBranch.upperName, {name, reverse: true})
 				}
 			}
 		}
@@ -192,6 +192,10 @@ if (botname === '__TEST__') {
 				}
 
 				source!.edgeProperties.set(targetName, edge)
+
+				if (edge.branchspec) {
+					source!.branchspec.set(targetName, {name: edge.branchspec, reverse: false})
+				}
 			}
 		}
 
diff --git a/Engine/Extras/RoboMerge/v3/src/robo/targets.ts b/Engine/Extras/RoboMerge/v3/src/robo/targets.ts
--- a/Engine/Extras/RoboMerge/v3/src/robo/targets.ts
+++ b/Engine/Extras/RoboMerge/v3/src/robo/targets.ts
@@ -678,9 +678,6 @@ function computeTargetsImpl(
 		const targetBranch = edgeTargetBranch(initialEdge)
 		const furtherEdgeBranches = furtherEdges.map(edgeTargetBranch)
 		for (const skipTarget of skipTargets) {
-			console.log('skip stuff!!!', initialEdge.source.debugName, skipTarget.name, furtherEdgeBranches)
-			console.log('skip stuff!!!', targetBranch.forcedDownstream)
-			console.log('skip stuff!!!', furtherEdgeBranches.map(b => b.forcedDownstream))
 			if (furtherEdgeBranches.indexOf(skipTarget) < 0
 				&&
 
