[NO NON-PII COMMIT MESSAGE]

diff --git a/Engine/Extras/RoboMerge/v3/src/robo/targets.ts b/Engine/Extras/RoboMerge/v3/src/robo/targets.ts
--- a/Engine/Extras/RoboMerge/v3/src/robo/targets.ts
+++ b/Engine/Extras/RoboMerge/v3/src/robo/targets.ts
@@ -309,7 +309,7 @@ export class DescriptionParser {
 				}
 
 				// ignore bare ROBOMERGE tags if we're not the default bot (should not affect default flow)
-				if (!command.bot) {
+				if (!command.bot && !isFlag(command.param)) {
 					if (this.isDefaultBot) {
 						this.useDefaultFlow = false
 					}
