[NO NON-PII COMMIT MESSAGE]

diff --git a/Engine/Extras/RoboMerge/v3/src/common/perforce.ts b/Engine/Extras/RoboMerge/v3/src/common/perforce.ts
--- a/Engine/Extras/RoboMerge/v3/src/common/perforce.ts
+++ b/Engine/Extras/RoboMerge/v3/src/common/perforce.ts
@@ -567,8 +567,12 @@ export class PerforceContext {
 			args.push('-a')
 		}
 
+		let opts: ExecZtagOpts = { multiline: true }
+		if (edgeServerAddress && edgeServerAddress !== 'commit') {
+			opts.edgeServerAddress = edgeServerAddress
+		}
 
-		let parsedClients = await this._execP4Ztag(null, args, { multiline: true, edgeServerAddress });
+		let parsedClients = await this._execP4Ztag(null, args, opts);
 		let workspaces = [];
 		for (let clientDef of parsedClients) {
 			if (clientDef.client) {
diff --git a/Engine/Extras/RoboMerge/v3/src/robo/edgebot.ts b/Engine/Extras/RoboMerge/v3/src/robo/edgebot.ts
--- a/Engine/Extras/RoboMerge/v3/src/robo/edgebot.ts
+++ b/Engine/Extras/RoboMerge/v3/src/robo/edgebot.ts
@@ -828,7 +828,9 @@ class EdgeBotImpl extends PerforceStatefulBot {
 		// Find a suitable workspace from one of the owner's workspaces
 		else {
 			// use p4.find_workspaces to find a workspace (owned by the user) for this change if this is a stream branch
-			const workspaces: ClientSpec[] = await this.p4.find_workspaces(owner)
+			// not supporting edge servers yet - really need to send a separate shelf request instead
+			// so specify to only look for workspaces on the commit server
+			const workspaces: ClientSpec[] = await this.p4.find_workspaces(owner, 'commit')
 
 			if (workspaces.length > 0) {
 				// default to the first workspace
