Have revertAndDelete take a workspace override.

The method took an edgeServer override but could end up with the wrong clientspec.

Should fix user requested shelves on an edge and robomerge trying to revert using the wrong clientspec on the edge server

#rnx

diff --git a/Engine/Extras/RoboMerge/v3/src/robo/edgebot.ts b/Engine/Extras/RoboMerge/v3/src/robo/edgebot.ts
--- a/Engine/Extras/RoboMerge/v3/src/robo/edgebot.ts
+++ b/Engine/Extras/RoboMerge/v3/src/robo/edgebot.ts
@@ -307,7 +307,7 @@ class EdgeBotImpl extends PerforceStatefulBot {
 		}
 
 		this.edgeBotLogger.info(`${logMessage}. Reverting ${pending.newCl}.`)
-		await this.revertAndDelete(pending.newCl)
+		await this.revertAndDelete(coercePerforceWorkspace(this.targetBranch.workspace)!.name, pending.newCl)
 
 		let pauseDurationSeconds = FAILED_CHANGELIST_PAUSE_TIMEOUT_SECONDS
 		// if we have a target, make sure pause duration is at least 2x duration of failed integration
@@ -557,7 +557,7 @@ class EdgeBotImpl extends PerforceStatefulBot {
 
 		// Revert attempt
 		if (pending.newCl > 0) {
-			await this.revertAndDelete(pending.newCl, edgeServerAddress)
+			await this.revertAndDelete(info.targetWorkspaceOverride, pending.newCl, edgeServerAddress)
 			pending.newCl = -1
 		}
 		
@@ -574,19 +574,19 @@ class EdgeBotImpl extends PerforceStatefulBot {
 		return new EdgeIntegrationDetails('error', errors.join('\n'))
 	}
 
-	private async revertAndDelete(cl: number, edgeServerAddress?: string) {
+	private async revertAndDelete(roboWorkspace: string, cl: number, edgeServerAddress?: string) {
 		this._log_action(`Reverting CL ${cl}`);
-		await this.p4.revert(this.targetBranch.workspace, cl, [], edgeServerAddress);
+		await this.p4.revert(roboWorkspace, cl, [], edgeServerAddress);
 
 		this._log_action(`Deleting CL ${cl}`);
-		await this.p4.deleteCl(this.targetBranch.workspace, cl, edgeServerAddress);
+		await this.p4.deleteCl(roboWorkspace, cl, edgeServerAddress);
 	}
 
 	async revertPendingCLWithShelf(client: string, change: number, userContext: string) {
 		this._log_action(` ${userContext} - Deleting shelf from and reverting ${change}`)
 		try {
 			this.edgeBotLogger.info(await this.p4.delete_shelved(client, change))
-			await this.revertAndDelete(change)
+			await this.revertAndDelete(client, change)
 		}
 		catch (err) {
 			this.edgeBotLogger.printException(err, `${userContext} - Error reverting ${change}`)
