[NO NON-PII COMMIT MESSAGE]

diff --git a/Engine/Extras/RoboMerge/v3/src/robo/edgebot.ts b/Engine/Extras/RoboMerge/v3/src/robo/edgebot.ts
--- a/Engine/Extras/RoboMerge/v3/src/robo/edgebot.ts
+++ b/Engine/Extras/RoboMerge/v3/src/robo/edgebot.ts
@@ -830,7 +830,7 @@ class EdgeBotImpl extends PerforceStatefulBot {
 			// use p4.find_workspaces to find a workspace (owned by the user) for this change if this is a stream branch
 			// not supporting edge servers yet - really need to send a separate shelf request instead
 			// so specify to only look for workspaces on the commit server
-			const workspaces: ClientSpec[] = await this.p4.find_workspaces(owner, 'commit')
+			const workspaces: ClientSpec[] = await p4util.getWorkspacesForUser(this.p4, owner)
 
 			if (workspaces.length > 0) {
 				// default to the first workspace
diff --git a/Engine/Extras/RoboMerge/v3/src/robo/ipc.ts b/Engine/Extras/RoboMerge/v3/src/robo/ipc.ts
--- a/Engine/Extras/RoboMerge/v3/src/robo/ipc.ts
+++ b/Engine/Extras/RoboMerge/v3/src/robo/ipc.ts
@@ -10,6 +10,7 @@ import { RoboMerge } from './robo';
 import { roboAnalytics } from './roboanalytics';
 import { Status } from './status';
 import { getPreview } from './preview'
+import * as p4util from '../common/p4util';
 
 // Reminder that RoboMerge variable is null in this file! ${RoboMerge}
 
@@ -143,7 +144,7 @@ export class IPC {
 	}
 
 	private async getWorkspaces(username: string): Promise<OperationReturnType> {
-		const workspaces = await this.robo.getWorkspacesForUser(username)
+		const workspaces = await p4util.getWorkspacesForUser(this.robo.p4, username)
 		
 		if (!workspaces) {
 			return {
diff --git a/Engine/Extras/RoboMerge/v3/src/robo/robo.ts b/Engine/Extras/RoboMerge/v3/src/robo/robo.ts
--- a/Engine/Extras/RoboMerge/v3/src/robo/robo.ts
+++ b/Engine/Extras/RoboMerge/v3/src/robo/robo.ts
@@ -34,11 +34,6 @@ import * as Preview from './preview'
 const roboStartupLogger = new ContextualLogger('Robo Startup')
 VersionReader.init(roboStartupLogger)
 
-
-const USER_WORKSPACE_EXCLUDE_PATTERNS: (RegExp | string)[] = [
-	'horde-p4bridge-'
-]
-
 // I seem to have broken this
 const DEBUG_SKIP_BRANCH_SETUP = false;
 
@@ -219,11 +214,6 @@ export class RoboMerge {
 		return graphBot && graphBot.settings.object
 	}
 
-	async getWorkspacesForUser(user: string) {
-		return (await this.p4.find_workspaces(user))
-			.filter(ws => !USER_WORKSPACE_EXCLUDE_PATTERNS.some(entry => ws.client.match(entry)))
-	}
-
 	stop() {
 		roboAnalytics!.stop()
 	}
@@ -315,8 +305,8 @@ function _checkForAutoPauseBots(branches: Branch[], logger: ContextualLogger) {
 		let paused = 0
 		for (const branch of branches) {
 			if (branch.bot) {
-				branch.bot.pause('Paused due to command line arguments or environment variables', 'robomerge')
-				++paused
+				// branch.bot.pause('Paused due to command line arguments or environment variables', 'robomerge')
+				// ++paused
 
 				// This should always be true
 				if (branch.bot instanceof NodeBot) {
