[NO NON-PII COMMIT MESSAGE]

diff --git a/Engine/Extras/RoboMerge/v3/src/common/p4util.ts b/Engine/Extras/RoboMerge/v3/src/common/p4util.ts
--- a/Engine/Extras/RoboMerge/v3/src/common/p4util.ts
+++ b/Engine/Extras/RoboMerge/v3/src/common/p4util.ts
@@ -3,6 +3,10 @@
 import { ContextualLogger } from './logger';
 import { Change, coercePerforceWorkspace, OpenedFileRecord, PerforceContext, RoboWorkspace } from './perforce';
 
+const USER_WORKSPACE_EXCLUDE_PATTERNS: (RegExp | string)[] = [
+	'horde-p4bridge-'
+]
+
 ///////////////////
 // Error reporting
 /*function _reportFilesStillToResolve(changeNum: number, needsResolve: Object[]) {
@@ -217,3 +221,9 @@ export async function cleanWorkspaces(p4: PerforceContext, workspaces: [string,
 	p4utilsLogger.info('Resetting all workspaces to revision 0')
 	await Promise.all(workspaces.map(([name, root]) => p4.sync(name, root + '#0', undefined, edgeServerAddress)))
 }
+
+export async function getWorkspacesForUser(p4: PerforceContext, user: string) {
+	return (await p4.find_workspaces(user))
+		.filter(ws => !USER_WORKSPACE_EXCLUDE_PATTERNS.some(entry => ws.client.match(entry)))
+}
+
