[NO NON-PII COMMIT MESSAGE]

diff --git a/Engine/Extras/RoboMerge/v3/src/common/p4util.ts b/Engine/Extras/RoboMerge/v3/src/common/p4util.ts
--- a/Engine/Extras/RoboMerge/v3/src/common/p4util.ts
+++ b/Engine/Extras/RoboMerge/v3/src/common/p4util.ts
@@ -222,8 +222,8 @@ export async function cleanWorkspaces(p4: PerforceContext, workspaces: [string,
 	await Promise.all(workspaces.map(([name, root]) => p4.sync(name, root + '#0', undefined, edgeServerAddress)))
 }
 
-export async function getWorkspacesForUser(p4: PerforceContext, user: string) {
-	return (await p4.find_workspaces(user))
+export async function getWorkspacesForUser(p4: PerforceContext, user: string, edgeServerAddress?: string) {
+	return (await p4.find_workspaces(user, edgeServerAddress))
 		.filter(ws => !USER_WORKSPACE_EXCLUDE_PATTERNS.some(entry => ws.client.match(entry)))
 }
 
diff --git a/Engine/Extras/RoboMerge/v3/src/robo/edgebot.ts b/Engine/Extras/RoboMerge/v3/src/robo/edgebot.ts
--- a/Engine/Extras/RoboMerge/v3/src/robo/edgebot.ts
+++ b/Engine/Extras/RoboMerge/v3/src/robo/edgebot.ts
@@ -830,7 +830,7 @@ class EdgeBotImpl extends PerforceStatefulBot {
 			// use p4.find_workspaces to find a workspace (owned by the user) for this change if this is a stream branch
 			// not supporting edge servers yet - really need to send a separate shelf request instead
 			// so specify to only look for workspaces on the commit server
-			const workspaces: ClientSpec[] = await p4util.getWorkspacesForUser(this.p4, owner)
+			const workspaces: ClientSpec[] = await p4util.getWorkspacesForUser(this.p4, owner, 'commit')
 
 			if (workspaces.length > 0) {
 				// default to the first workspace
