#rnx

diff --git a/Engine/Extras/RoboMerge/v3/README.md b/Engine/Extras/RoboMerge/v3/README.md
--- a/Engine/Extras/RoboMerge/v3/README.md
+++ b/Engine/Extras/RoboMerge/v3/README.md
@@ -36,7 +36,7 @@ RoboMerge will perform no further merges on the stream until the conflict has be
 
 ### Example commands
 
-![Example Robomerge Graph](/img/RM-Example-Graph.png)
+![Example Robomerge Graph](/public/images/RM-Example-Graph.png)
 
 A common pattern is to configure RoboMerge to merge commits to Main into release and feature branches. In the above example, commits to **Release-5.0** are _automatically_ merged up to **Release-6.0**, and then up to **Main**. The merge paths for each of our development streams can be seen on the [RoboMerge web page](/).
 
