[NO NON-PII COMMIT MESSAGE]

diff --git a/Engine/Extras/RoboMerge/v3/src/robo/notifications.ts b/Engine/Extras/RoboMerge/v3/src/robo/notifications.ts
--- a/Engine/Extras/RoboMerge/v3/src/robo/notifications.ts
+++ b/Engine/Extras/RoboMerge/v3/src/robo/notifications.ts
@@ -748,16 +748,8 @@ export class BotNotifications implements BotEventHandler {
 			if (message.fields) {
 				const newFields: SlackMessageField[] = []
 				for (const field of message.fields) {
-					switch (field.title) {
-						case CHANGELIST_FIELD_TITLE:
-							if (newClDesc) {
-								field.value = newClDesc
-							}
-							break
-
-						case ACKNOWLEDGED_FIELD_TITLE:
-							// skip add
-							continue
+					if (field.title === CHANGELIST_FIELD_TITLE && newClDesc) {
+						field.value = newClDesc
 					}
 					newFields.push(field)
 				}
