[NO NON-PII COMMIT MESSAGE]

diff --git a/Engine/Extras/RoboMerge/v3/src/common/analytics.ts b/Engine/Extras/RoboMerge/v3/src/common/analytics.ts
--- a/Engine/Extras/RoboMerge/v3/src/common/analytics.ts
+++ b/Engine/Extras/RoboMerge/v3/src/common/analytics.ts
@@ -115,6 +115,12 @@ export class Analytics {
 
 		this.memUsage.clear()
 
+		for (const [procName, usage] of this.diskUsage) {
+			lines.push(this._linePrefix + `event=disk,proc=${procName} value=${usage}`)
+		}
+
+		this.diskUsage.clear()
+
 		const branchesRequestsDelta = this.branchesRequests - this.lastReportedBranchesRequests
 		this.lastReportedBranchesRequests = this.branchesRequests
 		if (branchesRequestsDelta > 0) {
@@ -176,6 +182,10 @@ export class Analytics {
 		this.memUsage.set(procName, usage)
 	}
 
+	reportDiskUsage(procName: string, usage: number) {
+		this.diskUsage.set(procName, usage)
+	}
+
 	reportBranchesRequests(requests: number) {
 		this.branchesRequests = requests
 	}
@@ -209,6 +219,7 @@ export class Analytics {
 	}
 
 	private memUsage = new Map<string, number>()
+	private diskUsage = new Map<string, number>()
 	private branchesRequests = 0
 	private lastReportedBranchesRequests = 0
 
diff --git a/Engine/Extras/RoboMerge/v3/src/robo/watchdog.ts b/Engine/Extras/RoboMerge/v3/src/robo/watchdog.ts
--- a/Engine/Extras/RoboMerge/v3/src/robo/watchdog.ts
+++ b/Engine/Extras/RoboMerge/v3/src/robo/watchdog.ts
@@ -2,7 +2,7 @@
 
 import * as Sentry from '@sentry/node';
 import { OnUncaughtException, OnUnhandledRejection } from '@sentry/node/dist/integrations';
-import { ChildProcess, fork, ForkOptions } from 'child_process';
+import { ChildProcess, fork, ForkOptions, spawn } from 'child_process';
 import * as fs from 'fs';
 import * as os from 'os';
 import { Analytics } from '../common/analytics';
@@ -126,6 +126,15 @@ class Watchdog {
 		this.respawnTimer = null
 		this.memUsageTimer = setInterval(() => {
 			if (this.analytics) {
+				spawn('du', ['-s', '/src']).stdout.on('data', data => {
+					const match = data.match(/size:\s*(\d+)/)
+					if (match) {
+						const sizeBytes = parseInt(match[1])
+						if (!isNaN(sizeBytes)) {
+							this.analytics.reportDiskUsage('watchdog', sizeBytes)
+						}
+					}
+				})
 				this.analytics.reportMemoryUsage('watchdog', process.memoryUsage().heapUsed)
 				this.analytics.reportBranchesRequests(branchesRequests)
 			}
