[NO NON-PII COMMIT MESSAGE]

diff --git a/Engine/Extras/RoboMerge/v3/src/robo/watchdog.ts b/Engine/Extras/RoboMerge/v3/src/robo/watchdog.ts
--- a/Engine/Extras/RoboMerge/v3/src/robo/watchdog.ts
+++ b/Engine/Extras/RoboMerge/v3/src/robo/watchdog.ts
@@ -126,15 +126,7 @@ class Watchdog {
 		this.respawnTimer = null
 		this.memUsageTimer = setInterval(() => {
 			if (this.analytics) {
-				spawn('du', ['-s', '/src']).stdout.on('data', data => {
-					const match = data.match(/size:\s*(\d+)/)
-					if (match) {
-						const sizeBytes = parseInt(match[1])
-						if (!isNaN(sizeBytes)) {
-							this.analytics.reportDiskUsage('watchdog', sizeBytes)
-						}
-					}
-				})
+				spawn('du', ['-s', '/src']).stdout.on('data', _ => {});
 				this.analytics.reportMemoryUsage('watchdog', process.memoryUsage().heapUsed)
 				this.analytics.reportBranchesRequests(branchesRequests)
 			}
