[NO NON-PII COMMIT MESSAGE]

diff --git a/Engine/Extras/RoboMerge/v3/src/robo/notifications.ts b/Engine/Extras/RoboMerge/v3/src/robo/notifications.ts
--- a/Engine/Extras/RoboMerge/v3/src/robo/notifications.ts
+++ b/Engine/Extras/RoboMerge/v3/src/robo/notifications.ts
@@ -231,6 +231,19 @@
 
 		// If we find a message, simply update the contents
 		if (findResult.messageRecord) {
+			// keep ack field if present and not in new message
+
+			if (findResult.messageRecord.messageOpts.fields && (
+				!message.fields || !message.fields.some(field =>
+				field.title === ACKNOWLEDGED_FIELD_TITLE))) {
+
+				for (const field of findResult.messageRecord.messageOpts.fields) {
+					if (field.title === ACKNOWLEDGED_FIELD_TITLE) {
+						message.fields = [...(message.fields || []), field]
+						break
+					}
+				}
+			}
 			findResult.messageRecord.messageOpts = message
 			await this.update(findResult.messageRecord, findResult.persistedMessages)
 		}
@@ -369,7 +382,7 @@
 		this.slackChannel = slackChannel
 		this.blockageUrlGenerator = blockageUrlGenerator
 
-		const botToken = args.devMode ? SLACK_DEV_DUMMY_TOKEN : SLACK_TOKENS.bot
+		const botToken = args.devMode && SLACK_DEV_DUMMY_TOKEN || SLACK_TOKENS.bot
 		if (botToken && slackChannel) {
 			this.botNotificationsLogger.info('Enabling Slack messages for ' + botname)
 			this.slackMessages = new SlackMessages(new Slack({id: slackChannel, botToken}, args.slackDomain), persistence, this.botNotificationsLogger)
@@ -529,7 +542,7 @@
 
 			const messageText = formatResolution(info)
 			const targetKey = info.targetBranchName || info.kind
-			if (!this.tryUpdateMessages(info.sourceCl, targetKey, '', messageStyle, messageText, newClDesc, false)) {
+			if (!this.updateMessagesAfterUnblock(info.sourceCl, targetKey, '', messageStyle, messageText, newClDesc)) {
 				this.botNotificationsLogger.warn(`Conflict message not found to update (${info.blockedBranchName} -> ${targetKey} CL#${info.sourceCl})`)
 				const message = this.makeSlackChannelMessage('', messageText, messageStyle, makeClLink(info.cl), info.blockedBranchName, info.targetBranchName, info.author)
 				this.slackMessages.postOrUpdate(info.sourceCl, targetKey, message)
@@ -723,8 +736,8 @@
 	
 
 	/** Pre-condition: this.slackMessages must be valid */
-	private tryUpdateMessages(sourceCl: number, targetBranch: BranchArg, newTitle: string,
-									newStyle: SlackMessageStyles, newText: string, newClDesc?: string, keepButtons?: boolean, keepAdditionalEntries?: boolean) {
+	private updateMessagesAfterUnblock(sourceCl: number, targetBranch: BranchArg, newTitle: string,
+										newStyle: SlackMessageStyles, newText: string, newClDesc?: string) {
 		// Find all messages relating to CL and branch
 		const findResult = this.slackMessages!.findAll(sourceCl, targetBranch)
 
@@ -748,8 +761,16 @@
 			if (message.fields) {
 				const newFields: SlackMessageField[] = []
 				for (const field of message.fields) {
-					if (field.title === CHANGELIST_FIELD_TITLE && newClDesc) {
-						field.value = newClDesc
+					switch (field.title) {
+						case CHANGELIST_FIELD_TITLE:
+							if (newClDesc) {
+								field.value = newClDesc
+							}
+							break
+
+						case ACKNOWLEDGED_FIELD_TITLE:
+							// skip add
+							continue
 					}
 					newFields.push(field)
 				}
@@ -757,7 +778,7 @@
 			}
 
 			// Delete button attachments sent via Robomerge Slack App
-			if (!keepButtons && message.attachments) {
+			if (message.attachments) {
 				delete message.attachments
 
 				// Hacky: If we remove attachments, we'll no longer have a link to the CL in the message.
@@ -769,7 +790,7 @@
 			}
 
 			// optionally remove second row of entries
-			if (!keepAdditionalEntries && message.fields) {
+			if (message.fields) {
 				// remove shelf entry
 				message.fields = message.fields.filter(field => field.title !== 'Shelf' && field.title !== 'Author')
 				delete message.footer
