[NO NON-PII COMMIT MESSAGE]

diff --git a/Engine/Extras/RoboMerge/v3/src/robo/watchdog.ts b/Engine/Extras/RoboMerge/v3/src/robo/watchdog.ts
--- a/Engine/Extras/RoboMerge/v3/src/robo/watchdog.ts
+++ b/Engine/Extras/RoboMerge/v3/src/robo/watchdog.ts
@@ -13,7 +13,7 @@ import { VersionReader } from '../common/version';
 import { branchesRequests, RoboServer } from './roboserver';
 import { Session } from './session';
 
-const tlsKeyFilename = 'rm-2021-05.key'
+const tlsKeyFilename = 'rm-2022-05.key'
 
 // Begin by intializing our logger and version reader
 const watchdogStartupLogger = new ContextualLogger('Watchdog Startup')
