[NO NON-PII COMMIT MESSAGE]

diff --git a/Engine/Extras/RoboMerge/v3/README.md b/Engine/Extras/RoboMerge/v3/README.md
--- a/Engine/Extras/RoboMerge/v3/README.md
+++ b/Engine/Extras/RoboMerge/v3/README.md
@@ -101,7 +101,7 @@ for both node and edge.
 |`enabled`                  |ne|`true`      | |If false, pretends node/edge doesn't exist| |
 |`forcePause`               |ne|            | |If flag set, pause - applies each restart| |
 |`additionalSlackChannelForBlockages` |e |  | |Single extra Slack channel               | |
-|`blockAssetTargets`        |e |`false`     | |Reject integrations containing assets    | |
+|`blockAssetFlow`           |n |`[]`        | |Reject any integrations containing assets to these targets   | |
 |`defaultIntegrationMethod` |e |normal      | |For edigrate                             | |
 |`disallowSkip`             |e |            | |Remove skip option from UI               | |
 |`emailOnBlockage`          |e |`true`      | |Email owners of conflicts?               | |
