[NO NON-PII COMMIT MESSAGE]

diff --git a/Engine/Extras/RoboMerge/v3/src/new/graph.ts b/Engine/Extras/RoboMerge/v3/src/new/graph.ts
--- a/Engine/Extras/RoboMerge/v3/src/new/graph.ts
+++ b/Engine/Extras/RoboMerge/v3/src/new/graph.ts
@@ -440,7 +440,7 @@ class Test {
 
 	private fillTargetsMap(targets: string) {
 
-		const targetBits = targets.split('-') // up to three section in target (1st) string: normal, skip, null
+		const targetBits = targets.split(':') // up to three section in target (1st) string: normal, skip, null
 
 		for (const targetChar of targetBits[0]) {
 			this.targets.set(this.getNode(targetChar), 'normal')
@@ -536,12 +536,24 @@ Test format:
 graph definition e.g.: ['b', 'c', ''] means a->b, b->c and d also exists
 expected: ;-separated <direct>:<indirect>
 
-4: A -> B -> c		c-b expected to produce D:C (A->D->C)
+1: A -> B -> C		c:b expected to produce null (c unreachable)
+2: A => B -> C		c:b expected to produce null (c unreachable)
+3: A -> B -> C	_:c should produce '' (skip shouldn't affect route)
+
+4/5: dev/release set-up
+
+
+6: A -> B -> c		c:d expected to produce B:C (A->B->C)
+	 -> D -> c
+
+7: A -> B -> c		c:b expected to produce D:C (A->D->C)
 	 => D -> c
 
-5: A => B => C		-c expected to produce B:-C
+8: A => B => C		_:c expected to produce B:-C
 	 -> D
 
+9: A => B => C		_:c:b expected to produce !B:-C
+
 
 */
 
@@ -550,37 +562,42 @@ expected: ;-separated <direct>:<indirect>
 	const unitTestLogger = parentLogger.createChild('Graph')
 	const tests: [string[], string | null][] = [
 
-		[['-c', 'a', 'b', 'c'], ''],
-		[['c-b', 'a', 		'b', 'c', ''],										null],
-		[['c-b', 'a', 		'B', 'c', ''],										null],
-		[['c-d', 'a', 		'bd', 'c', '', 'c'],								'B:C'],
-		[['c-b', 'a', 		'bD', 'c', '', 'c'],								'D:C'], //4
-		[['-c', 'a',		'Bd', 'C', '', ''],									'B:-C'],
-		[['-c-b', 'a',		'B', 'C', ''],										'!B:-C'],
-		[['b', 'a', 		'b', 'C', ''],										'B:'],
+		[[':c', 'a', 'b', 'c'], ''],
+		[['c:b', 'a', 		'b', 'c', ''],										null],
+		[['c:b', 'a', 		'B', 'c', ''],										null],
+		[['b:c', 'a', 		'b', 'c', ''],										''],
+
+		// usual dev/release set-up (b is most recent release)
+		[['c', 'a',			'bDE', 'Ac', 'B', 'a', 'a'],						'B:C'],
+		[[':c', 'a',		'bDE', 'Ac', 'B', 'a', 'a'],						''], // 5
+
+		[['c:d', 'a', 		'bd', 'c', '', 'c'],								'B:C'],
+		[['c:b', 'a', 		'bD', 'c', '', 'c'],								'D:C'],
+		[[':c', 'a',		'Bd', 'C', '', ''],									'B:-C'],
+		[[':c:b', 'a',		'B', 'C', ''],										'!B:-C'],
+		[['b', 'a', 		'b', 'C', ''],										'B:'], // 10
 		[['f', 'b',			'bcd', 'a', 'a', 'Ae', 'Df', 'Eg', 'F'],			'A:DEF'],
 		[['b', 'f',			'bcd', 'a', 'a', 'Ae', 'Df', 'Eg', 'F'],			'E:DAB'],
 		[['a', 'g',			'bcd', 'a', 'a', 'Ae', 'Df', 'Eg', 'F'],			'F:EDA'],
-		[['-a', 'g',		'bcd', 'a', 'a', 'Ae', 'Df', 'Eg', 'F'],			'F:ED-A'],
+		[[':a', 'g',		'bcd', 'a', 'a', 'Ae', 'Df', 'Eg', 'F'],			'F:ED-A'],
 		[['de', 'a',		'b', 'c', 'de', '', ''],							'B:CED'],
 		[['de', 'h',		'hc', 'hd', 'deFg', 'hbc', 'c', 'c', 'c', 'abd'],	'D:CE'],
 		[['de', 'h',		'bd', 'hc', 'hd', 'eFg', 'hbc', 'c', 'c', 'c'],		'C:DE'],
 		[['db', 'h',		'abd', 'hc', 'hd', 'deFg', 'hbc', 'c', 'c', 'c'],	'C:DEB'],
 		[['cg', 'a',		'be', 'acd', 'b', 'b', 'aFg', 'e', 'e'],			'B:C;E:G'],
 		[['cf', 'a', 		'bE', 'acd', 'b', 'b', 'aFg', 'e', 'e'],			'B:C;E:F'],
-		[['c-e', 'a', 		'bE', 'acd', 'b', 'b', 'aFg', 'e', 'e'],			'B:C'],
-		[['cf-b', 'a', 		'bE', 'acd', 'b', 'b', 'aFg', 'e', 'e'],			null],
+		[['c:e', 'a', 		'bE', 'acd', 'b', 'b', 'aFg', 'e', 'e'],			'B:C'],
+		[['cf:b', 'a', 		'bE', 'acd', 'b', 'b', 'aFg', 'e', 'e'],			null],
 		[['fg', 'h',		'abd', 'hC', 'hd', 'defg', 'hbc', 'c', 'c', 'c'],	'C:DGF'],
 		[['dge', 'h',		'abd', 'hC', 'hd', 'defg', 'hbc', 'c', 'c', 'c'],	'C:DGE'],
 		[['e', 'i',			'abcd', 'Ie', 'if', 'iG', 'ih', 'a', 'b', 'c', 'D'],'D:GBE']
 	]
 
-	let ran = 0, success = 0, fail = 0, index = 0
-	for (const [testStr, expected] of tests /*.slice(3, 4)*/) {
+	let success = 0, fail = 0, ran = 0
+	for (const [testStr, expected] of tests /*/ .slice(6, 7) /**/) {
 		const test = new Test(testStr.slice(2))
 		const result = test.computeTargets(testStr[1], testStr[0])
 
-		++ran
 		let expectedOnFail
 
 		const succeeded = result.status === 'succeeded'
@@ -599,15 +616,17 @@ expected: ;-separated <direct>:<indirect>
 		else if (succeeded) {
 			expectedOnFail = 'fail'
 		}
+		else { console.log(ran, result.unreachable && result.unreachable.map(n => n.debugName)) }
+
+		++ran
 		if (expectedOnFail) {
 			++fail
-			unitTestLogger.warn(`Test ${colors.warn(index.toString().padStart(2))} failed:   ` +
+			unitTestLogger.warn(`Test ${colors.warn(ran.toString().padStart(2))} failed:   ` +
 				`${colors.warn(formattedResult.padStart(10))} vs ${colors.warn(expectedOnFail)}`)
 		}
 		else {
 			++success
 		}
-		++index
 	}
 	const message = `Graph tests: ran ${ran} tests, ${success} matched, ${fail} failed`
 	unitTestLogger.info((fail > 0 ? colors.error : colors.info)(message))
