[NO NON-PII COMMIT MESSAGE]

diff --git a/Engine/Extras/RoboMerge/v3/src/new/graph.ts b/Engine/Extras/RoboMerge/v3/src/new/graph.ts
--- a/Engine/Extras/RoboMerge/v3/src/new/graph.ts
+++ b/Engine/Extras/RoboMerge/v3/src/new/graph.ts
@@ -255,10 +255,12 @@ export class Graph {
 		const targetsToFind = new Set<Node>()
 		const skipNodes = new Set<Node>()
 		for (const [target, mergeMode] of requestedTargets.entries()) {
-			targetsToFind.add(target)
 			if (mergeMode === 'skip') {
 				skipNodes.add(target)
 			}
+			else {
+				targetsToFind.add(target)
+			}
 		}
 
 		// initial edge to follow, then subsequent
