[NO NON-PII COMMIT MESSAGE]

diff --git a/Engine/Extras/RoboMerge/v3/src/robo/targets.ts b/Engine/Extras/RoboMerge/v3/src/robo/targets.ts
--- a/Engine/Extras/RoboMerge/v3/src/robo/targets.ts
+++ b/Engine/Extras/RoboMerge/v3/src/robo/targets.ts
@@ -582,7 +582,7 @@ function computeTargetsImpl(
 	}
 
 	const branchGraph = sourceBranch.parent
-	const requestedMerges = [
+	const requestedMergesArray = [
 		...defaultTargets.map(name => [name, 'normal'] as [string, MergeMode]),
 		...ri.integrations.filter(([name, _]) => branchGraph.config.branchNamesToIgnore.indexOf(name.toUpperCase()) < 0)
 	]
@@ -593,7 +593,7 @@ function computeTargetsImpl(
 	const errors: string[] = []
 
 	// process parsed targets
-	for (const [targetName, mergeMode] of requestedMerges) {
+	for (const [targetName, mergeMode] of requestedMergesArray) {
 		// make sure the target exists
 		const targetBranch = branchGraph.getBranch(targetName)
 		if (!targetBranch) {
@@ -622,10 +622,10 @@ function computeTargetsImpl(
 	}
 
 // note: this allows multiple mentions of same branch, with flag of last mention taking priority. Could be an error instead
-	const allIntegrations = new Map<string, MergeMode>(requestedMerges)
+	const requestedMerges = new Map<string, MergeMode>(requestedMergesArray)
 
 	const requestedTargets = new Map<Node, MergeMode>()
-	for (const [targetName, mergeMode] of allIntegrations) {
+	for (const [targetName, mergeMode] of requestedMerges) {
 		const node = ubergraph.getNode(makeTargetName(botname, targetName))
 		if (node) {
 			if (node !== sourceNode) {
@@ -659,7 +659,7 @@ function computeTargetsImpl(
 
 	// handle multiple routes to skip targets
 	//	
-	const skipTargets: Branch[] = [...allIntegrations]
+	const skipTargets: Branch[] = [...requestedMerges]
 		.filter(([_, mergeMode]) => mergeMode === 'skip')
 		.map(([branchName, _]) => branchGraph.getBranch(branchName)!)
 
