[NO NON-PII COMMIT MESSAGE]

diff --git a/Engine/Extras/RoboMerge/v3/functional_tests/framework/src/GenericTargetTest.ts b/Engine/Extras/RoboMerge/v3/functional_tests/framework/src/GenericTargetTest.ts
new file mode 100644
--- /dev/null
+++ b/Engine/Extras/RoboMerge/v3/functional_tests/framework/src/GenericTargetTest.ts
@@ -0,0 +1,120 @@
+import { FunctionalTest, P4Util, RobomergeBranchSpec, Stream } from './framework'
+import { Perforce } from './test-perforce'
+
+
+// testDef example
+//    command | start |        edges         | expected
+// [['c:d',		'a',    'bd', 'c', '', 'c'],	'B:C']
+
+export class GenericTargetTest extends FunctionalTest {
+	// private streams: Stream[] = []
+
+
+	private streamMap = new Map<string, Stream>();
+	private branchMap = new Map<string, RobomergeBranchSpec>();
+	private source: string;
+	private command: string;
+
+	constructor(p4: Perforce, index: number, testDef: string[], private expected: string | null) {
+		super(p4, `TargetTest${index}`)
+
+		// copy code from Test constructor/computeTargets in graph.ts!
+		const targetsString = testDef[0]
+		this.source = testDef[1]
+		const nodeStrs = testDef.slice(2)
+
+		// name -> char (variable names same as Test.constructor)
+		for (let index = 0; index < nodeStrs.length; ++index) {
+			const name = String.fromCharCode('a'.charCodeAt(0) + index)
+			const branch = this.addBranch(name)
+			this.streamMap.set(name, index === 0
+				? {name, streamType: 'mainline'}
+				: {name, streamType: 'development', parent: 'a'})
+
+			for (const char of nodeStrs[index]) {
+				const targetName = char.toLowerCase()
+				this.addBranch(targetName)
+				branch.flowsTo!.push(this.fullBranchName(targetName))
+				if (char !== targetName) {
+					branch.forceFlowTo!.push(this.fullBranchName(targetName))
+				}
+			}
+		}
+
+		const commandBits: string[] = []
+		//fillTargetsMap
+		const targetBits = targetsString.split(':') // up to three section in target (1st) string: normal, skip, null
+		for (const targetChar of targetBits[0]) {
+			commandBits.push(this.fullBranchName(targetChar))
+		}
+
+		if (targetBits.length > 1) {
+			for (const targetChar of targetBits[1]) {
+				commandBits.push('-' + this.fullBranchName(targetChar))
+			}
+		}
+
+		if (targetBits.length > 2) {
+			for (const targetChar of targetBits[2]) {
+				commandBits.push('!' + this.fullBranchName(targetChar))
+			}
+		}
+
+		this.command = commandBits.join(' ')
+	}
+
+	async setup() {
+		const streams = [...this.streamMap.values()]
+		await this.p4.depot('stream', this.depotSpec())
+		await this.createStreamsAndWorkspaces(streams)
+		await P4Util.addFileAndSubmit(this.getClient('a'), 'test.txt', 'initial content')
+
+		const desc = 'Initial population'
+				
+		await Promise.all(streams
+			.filter(s => s.name !== 'a')
+			.map(s => this.p4.populate(this.getStreamPath(s.name), desc))
+		)
+	}
+
+	async run() {
+		const sourceClient = this.getClient(this.source)
+		if (this.source !== 'a') {
+			await sourceClient.sync()
+		}
+		await P4Util.editFileAndSubmit(sourceClient, 'test.txt', 'new content', this.command)
+	}
+
+	async verify() {
+		// at the moment, checking normal and null merges create a revision
+		// upper case in expected should be checking for null merge
+
+		if (this.expected || this.expected === '') {
+			return Promise.all([...this.streamMap.values()].map(s => 
+				this.checkHeadRevision(s.name, 'test.txt',
+					1 +
+					(s.name === this.source ? 1 : 0) +
+					(this.expected!.toLowerCase().indexOf(s.name) >= 0 ? 1 : 0))))
+		}
+		else {
+			return this.ensureBlocked(this.source)
+		}
+	}
+
+	getBranches() {
+		return [...this.branchMap.values()]
+
+	}
+
+	allowSyntaxErrors() {
+		return true
+	}
+	private addBranch(name: string): RobomergeBranchSpec {
+		let branch = this.branchMap.get(name)
+		if (!branch) {
+			branch = this.makeBranchDef(name, [], false)
+			this.branchMap.set(name, branch)
+		}
+		return branch
+	}
+}
diff --git a/Engine/Extras/RoboMerge/v3/functional_tests/framework/src/tests.ts b/Engine/Extras/RoboMerge/v3/functional_tests/framework/src/tests.ts
--- a/Engine/Extras/RoboMerge/v3/functional_tests/framework/src/tests.ts
+++ b/Engine/Extras/RoboMerge/v3/functional_tests/framework/src/tests.ts
@@ -3,6 +3,7 @@ import { DEFAULT_BOT_SETTINGS, EdgeProperties, FunctionalTest, getRootDataClient
 				P4Util, RobomergeBranchSpec, ROBOMERGE_DOMAIN, retryWithBackoff } from './framework'
 import * as bent from 'bent'
 import { Perforce } from './test-perforce'
+import { GenericTargetTest } from './GenericTargetTest'
 import { BlockAssets } from './tests/block-assets'
 import { BlockIgnore } from './tests/block-ignore'
 import { ConfirmBinaryStomp } from './tests/confirm-binary-stomp'
@@ -42,6 +43,7 @@ import { TestGate } from './tests/test-gate'
 import { TestReconsider } from './tests/test-reconsider'
 import { TestEdgeReconsider } from './tests/test-edge-reconsider'
 import { TestTerminal } from './tests/test-terminal'
+import { UnreachableSkip } from './tests/unreachable-skip'
 
 import { CrossBotTest, CrossBotTest2, ComplexCrossBot, ComplexCrossBot2, ComplexCrossBot3 } from './tests/cross-bot'
 
@@ -64,7 +66,7 @@ function addTest(settings: BranchMapSettings, test: FunctionalTest) {
 	settings.macros = {...settings.macros, ...test.getMacros()}
 }
 
-async function addToRoboMerge(p4: Perforce, tests: FunctionalTest[]) {
+async function 	addToRoboMerge(p4: Perforce, tests: FunctionalTest[], targetTests: FunctionalTest[]) {
 	const rootClient = await getRootDataClient(p4, 'RoboMergeData_BranchMaps')
 
 	const botNames = ['ft1', 'ft2', 'ft3', 'ft4']
@@ -85,6 +87,17 @@ async function addToRoboMerge(p4: Perforce, tests: FunctionalTest[]) {
 		groupIndex = (groupIndex + 1) % botNames.length
 	}
 
+	const TARGET_TEST_BOT_NAME = 'targets'
+
+	const targetTestsBranchMapSettings = {branches: [], edges: [], macros: {}}
+	settings.push([TARGET_TEST_BOT_NAME, targetTestsBranchMapSettings])
+
+	for (const test of targetTests) {
+		test.botName = TARGET_TEST_BOT_NAME
+		addTest(targetTestsBranchMapSettings, test)
+		test.storeNodesAndEdges()
+	}
+
 	await Promise.all(settings.map(([botName, s]) => 
 		P4Util.addFile(rootClient, botName + '.branchmap.json', JSON.stringify({...DEFAULT_BOT_SETTINGS, ...s, slackChannel: botName,
 			aliases: [botName + '-alias', botName + '-alias2']}))
@@ -97,7 +110,7 @@ async function addToRoboMerge(p4: Perforce, tests: FunctionalTest[]) {
 
 	await retryWithBackoff('Waiting for all branchmaps to load', async () => {
 
-		for (const test of tests) {
+		for (const test of [...tests, ...targetTests]) {
 			for (const node of test.nodes) {
 				try {
 					await FunctionalTest.getBranchState(test.botName, node)
@@ -198,9 +211,51 @@ async function go() {
 		new StompWithAdd(p4),
 		new Ignore(p4),
 		new StompForwardingCommands(p4),
-		new MultipleRoutesToSkip(p4)
+		new MultipleRoutesToSkip(p4),
+		new UnreachableSkip(p4), // 45
 	]
 
+	const TARGET_TEST_DEFS: [string[], string | null][] = [
+		[[':c', 'a',		'b', 'c', ''],										''],
+		[['c:b', 'a', 		'b', 'c', ''],										null],
+		[['c:b', 'a', 		'B', 'c', ''],										null],
+		[[':c', 'a', 		'b', 'c', ''],										''],
+
+		[['c', 'a',			'bDE', 'Ac', 'B', 'a', 'a'],						'bcde'],
+		[[':c', 'a',		'bDE', 'Ac', 'B', 'a', 'a'],						'de'], // 5
+
+		[['c:d', 'a', 		'bd', 'c', '', 'c'],								'bc'],
+		[['c:b', 'a', 		'bD', 'c', '', 'c'],								'dc'],
+		[[':c', 'a',		'Bd', 'C', '', ''],									'b'],
+		[[':c:b', 'a',		'B', 'C', ''],										'B'], 
+		[['b', 'a', 		'b', 'C', ''],										'bc'], // 10
+
+// usual dev/release set-up (d-g are releases going back in time)
+		[['f', 'b',			'bcd', 'a', 'a', 'Ae', 'Df', 'Eg', 'F'],			'adef'], 
+		[['b', 'f',			'bcd', 'a', 'a', 'Ae', 'Df', 'Eg', 'F'],			'edab'],
+		[['a', 'g',			'bcd', 'a', 'a', 'Ae', 'Df', 'Eg', 'F'],			'feda'],
+		[[':a', 'g',		'bcd', 'a', 'a', 'Ae', 'Df', 'Eg', 'F'],			'def'],
+
+		[['de', 'a',		'b', 'c', 'de', '', ''],							'bced'], // 15
+		[['de', 'h',		'hc', 'hd', 'deFg', 'hbc', 'c', 'c', 'c', 'abd'],	'dcef'],
+		[['de', 'h',		'bd', 'hc', 'hd', 'eFg', 'hbc', 'c', 'c', 'c'],		'cdef'],
+		[['db', 'h',		'abd', 'hc', 'hd', 'deFg', 'hbc', 'c', 'c', 'c'],	'cdebf'],
+		[['cg', 'a',		'be', 'acd', 'b', 'b', 'aFg', 'e', 'e'],			'bcegf'],
+		[['cf', 'a', 		'bE', 'acd', 'b', 'b', 'aFg', 'e', 'e'],			'bcef'], // 20
+		[['c:e', 'a', 		'bE', 'acd', 'b', 'b', 'aFg', 'e', 'e'],			'bc'],
+		[['cf:b', 'a', 		'bE', 'acd', 'b', 'b', 'aFg', 'e', 'e'],			null], 
+		[['fg', 'h',		'abd', 'hC', 'hd', 'defg', 'hbc', 'c', 'c', 'c'],	'cdgf'],
+		[['dge', 'h',		'abd', 'hC', 'hd', 'defg', 'hbc', 'c', 'c', 'c'],	'cdge'],
+		[['e', 'i',			'abcd', 'Ie', 'if', 'iG', 'ih', 'a', 'b', 'c', 'D'],'dgbe']
+	]
+
+	let targetTestIndex = 0
+	const availableTargetTests = []
+	for (const [strs, expected] of TARGET_TEST_DEFS) {
+		availableTargetTests.push(new GenericTargetTest(p4, targetTestIndex++, strs, expected))
+	}
+
+
 	// const testToDebug = availableTests[30]
 	// testToDebug.botName = 'FT1'
 	// await testToDebug.run()
@@ -224,7 +279,9 @@ async function go() {
 	///////////////////////
 	// TESTS TO RUN 
 
-	const tests = /*/[availableTests[30]]/*/availableTests  /* .slice(34, 39)/**/
+	const specificTests = /*/[availableTests[17]]/*/availableTests  /*/ .slice(0, 0)/**/
+	const targetTests = /*/[availableTargetTests[7]] /*/ availableTargetTests /*/ .slice(8, 10) /**/
+	const tests = [...specificTests, ...targetTests]
 
 	//
 	///////////////////////
@@ -235,7 +292,7 @@ async function go() {
 	await Promise.all(tests.map(test => test.setup()))
 
 	console.log('Updating RoboMerge branchmaps')
-	await addToRoboMerge(p4, tests)
+	await addToRoboMerge(p4, specificTests, targetTests)
 
 	console.log('Running tests')
 	await Promise.all(tests.map(test => test.run()))
diff --git a/Engine/Extras/RoboMerge/v3/functional_tests/framework/src/tests/unreachable-skip.ts b/Engine/Extras/RoboMerge/v3/functional_tests/framework/src/tests/unreachable-skip.ts
new file mode 100644
--- /dev/null
+++ b/Engine/Extras/RoboMerge/v3/functional_tests/framework/src/tests/unreachable-skip.ts
@@ -0,0 +1,30 @@
+// Copyright Epic Games, Inc. All Rights Reserved.
+import { P4Util } from '../framework'
+import { MultipleDevAndReleaseTestBase } from '../MultipleDevAndReleaseTestBase'
+
+
+export class UnreachableSkip extends MultipleDevAndReleaseTestBase {
+
+	async setup() {
+		await this.p4.depot('stream', this.depotSpec())
+		await this.createStreamsAndWorkspaces()
+		await P4Util.addFileAndSubmit(this.getClient('Main'), 'test.txt', 'content')
+		await this.initialPopulate()
+	}
+
+	run() {
+		return P4Util.editFileAndSubmit(this.getClient('Main'), 'test.txt', 'updated', '-' + this.fullBranchName('Release-1.0'))
+	}
+
+	verify() {
+		// currently not erroring
+		// return this.ensureBlocked('Main')
+
+		return this.checkHeadRevision('Release-2.0', 'test.txt', 1)
+	}
+
+	allowSyntaxErrors() {
+		return true
+	}
+}
+
diff --git a/Engine/Extras/RoboMerge/v3/src/new/graph.ts b/Engine/Extras/RoboMerge/v3/src/new/graph.ts
--- a/Engine/Extras/RoboMerge/v3/src/new/graph.ts
+++ b/Engine/Extras/RoboMerge/v3/src/new/graph.ts
@@ -395,7 +395,7 @@ export type MergeMode = 'safe' | 'normal' | 'null' | 'clobber' | 'skip'
 
 class Test {
 	graph = new Graph
-	targets = new Map<Node, MergeMode>()
+	targets = new Map<Node, MergeMode>();
 
 	static readonly BOT_NAME = 'TEST' as BotName
 
@@ -538,7 +538,7 @@ Test format:
 graph definition e.g.: ['b', 'c', ''] means a->b, b->c and d also exists
 expected: ;-separated <direct>:<indirect>
 
-1: A -> B -> C		c:b expected to produce null (c unreachable)
+1: A -> B -> C		:c expected to produce '' (c unreachable, could alternatively be syntax error)
 2: A => B -> C		c:b expected to produce null (c unreachable)
 3: A -> B -> C	_:c should produce '' (skip shouldn't affect route)
 
@@ -551,37 +551,42 @@ expected: ;-separated <direct>:<indirect>
 7: A -> B -> c		c:b expected to produce D:C (A->D->C)
 	 => D -> c
 
-8: A => B => C		_:c expected to produce B:-C
+8: A => B => C		_:c expected to produce ''
 	 -> D
 
-9: A => B => C		_:c:b expected to produce !B:-C
+9: A => B => C		_:c:b expected to produce !B
 
 
+Note first level is treated specially: for the purposes of these tests, force flow from the start node is effectively ignored
 */
 
 
+// want to find a way of testing commented out ones - the relevant skip handling is done in targets.ts
+
 
 	const unitTestLogger = parentLogger.createChild('Graph')
 	const tests: [string[], string | null][] = [
 
-		[[':c', 'a', 'b', 'c'], ''],
+		[[':c', 'a',		'b', 'c', ''],										''],
 		[['c:b', 'a', 		'b', 'c', ''],										null],
 		[['c:b', 'a', 		'B', 'c', ''],										null],
-		[['b:c', 'a', 		'b', 'c', ''],										''],
+		[[':c', 'a', 		'b', 'c', ''],										''],
 
-		// usual dev/release set-up (b is most recent release)
 		[['c', 'a',			'bDE', 'Ac', 'B', 'a', 'a'],						'B:C'],
-		[[':c', 'a',		'bDE', 'Ac', 'B', 'a', 'a'],						''], // 5
+		[[':c', 'a',		'bDE', 'Ac', 'B', 'a', 'a'],						''], //DE'], // 5
 
 		[['c:d', 'a', 		'bd', 'c', '', 'c'],								'B:C'],
 		[['c:b', 'a', 		'bD', 'c', '', 'c'],								'D:C'],
-		// [[':c', 'a',		'Bd', 'C', '', ''],									'B:-C'],
-		// [[':c:b', 'a',		'B', 'C', ''],										'!B:-C'],
+		[[':c', 'a',		'Bd', 'C', '', ''],									''],
+		[[':c:b', 'a',		'B', 'C', ''],										'!B:'], // !B:-C'], skip is checked by functional test
 		[['b', 'a', 		'b', 'C', ''],										'B:'], // 10
+
+// usual dev/release set-up (d-g are releases going back in time)
 		[['f', 'b',			'bcd', 'a', 'a', 'Ae', 'Df', 'Eg', 'F'],			'A:DEF'],
 		[['b', 'f',			'bcd', 'a', 'a', 'Ae', 'Df', 'Eg', 'F'],			'E:DAB'],
 		[['a', 'g',			'bcd', 'a', 'a', 'Ae', 'Df', 'Eg', 'F'],			'F:EDA'],
-		// [[':a', 'g',		'bcd', 'a', 'a', 'Ae', 'Df', 'Eg', 'F'],			'F:ED-A'],
+		[[':a', 'g',		'bcd', 'a', 'a', 'Ae', 'Df', 'Eg', 'F'],			''], // taking a look at this one
+
 		[['de', 'a',		'b', 'c', 'de', '', ''],							'B:CED'],
 		[['de', 'h',		'hc', 'hd', 'deFg', 'hbc', 'c', 'c', 'c', 'abd'],	'D:CE'],
 		[['de', 'h',		'bd', 'hc', 'hd', 'eFg', 'hbc', 'c', 'c', 'c'],		'C:DE'],
@@ -596,35 +601,36 @@ expected: ;-separated <direct>:<indirect>
 	]
 
 	let success = 0, fail = 0, ran = 0
-	for (const [testStr, expected] of tests /*/ .slice(6, 7) /**/) {
+	for (const [testStr, expected] of tests /*/ .slice(5, 6) /**/) {
 		const test = new Test(testStr.slice(2))
 		const result = test.computeTargets(testStr[1], testStr[0])
 
-		let expectedOnFail
+		let expectedOnFail: string | null = null
 
 		const succeeded = result.status === 'succeeded'
 
 		const formattedResult = succeeded ? test.formatTestComputeTargetsResult(result.integrations!) : 'failed'
 		if (expected !== null) {
+			// not expecting an error
 			if (!succeeded) {
+				// store string we were expecting
 				expectedOnFail = expected
 			}
-			else {
-				if (formattedResult.toUpperCase().replace(/\s|,/g, '') !== expected) {
-					expectedOnFail = expected
-				}
+			else if (formattedResult.toUpperCase().replace(/\s|,/g, '') !== expected) {
+				expectedOnFail = expected
 			}
 		}
 		else if (succeeded) {
 			expectedOnFail = 'fail'
 		}
-		else { console.log(ran, result.unreachable && result.unreachable.map(n => n.debugName)) }
 
 		++ran
-		if (expectedOnFail) {
+		if (expectedOnFail === undefined) throw new Error('wut')
+		if (expectedOnFail !== null) {
 			++fail
+			const EMPTY = '<empty>'
 			unitTestLogger.warn(`Test ${colors.warn(ran.toString().padStart(2))} failed:   ` +
-				`${colors.warn(formattedResult.padStart(10))} vs ${colors.warn(expectedOnFail)}`)
+				`${colors.warn((formattedResult || EMPTY).padStart(10))} vs ${colors.warn(expectedOnFail || EMPTY)}`)
 		}
 		else {
 			++success
