[NO NON-PII COMMIT MESSAGE]

diff --git a/Engine/Extras/RoboMerge/v3/functional_tests/framework/src/framework.ts b/Engine/Extras/RoboMerge/v3/functional_tests/framework/src/framework.ts
--- a/Engine/Extras/RoboMerge/v3/functional_tests/framework/src/framework.ts
+++ b/Engine/Extras/RoboMerge/v3/functional_tests/framework/src/framework.ts
@@ -194,9 +194,12 @@ type EdgeOptionFields = {
 
 	implicitCommands: string[]
 
+	ignoreInCycleDetection: boolean
+
 	approval: {
 		description: string,
-		group: string
+		channelName: string
+		channelId: string
 	}
 }
 
@@ -309,10 +312,16 @@ export abstract class FunctionalTest {
 			name, aliases: [name + '_alias'],
 			streamName: stream,
 			flowsTo: to.map(str => this.fullBranchName(str)),
+			forceFlowTo: [],
 			forceAll: !!forceAll
 		}
 	}
 
+	protected makeEdgeProperties(from: string, to: string, props: Partial<EdgeProperties>): EdgeProperties {
+		return {from: this.fullBranchName(from), to: this.fullBranchName(to), ...props}
+	}
+
+
 	protected makeForceAllBranchDef(stream: string, to: string[]/*, additionalProps?: {[prop: string]: string}[]*/): RobomergeBranchSpec {
 		return this.makeBranchDef(stream, to, true)
 	}
@@ -571,7 +580,7 @@ export abstract class FunctionalTest {
 			throw new Error('edge must be blocked to stomp!')
 		}
 
-		const conflictCl = edgeState.conflictCl
+		const conflictCl = edgeState.conflict && edgeState.conflict.change
 		const sourceBranchName = this.fullBranchName(source)
 		const targetBranchName = this.fullBranchName(target)
 
@@ -598,7 +607,7 @@ export abstract class FunctionalTest {
 	}
 
 	async performStompRequest(source: string, target: string, edgeState: EdgeState) {
-		const conflictCl = edgeState.conflictCl
+		const conflictCl = edgeState.conflict && edgeState.conflict.change
 		const sourceBranchName = this.fullBranchName(source)
 		const targetBranchName = this.fullBranchName(target)
 
@@ -657,18 +666,30 @@ export abstract class FunctionalTest {
 		const edgeDisplayName = `${sourceStream} -> ${targetStream}`
 		this.info(`Ensuring conflict message sent for ${edgeDisplayName}`)
 		const edgeState = await this.getEdgeState('Release', 'Main')
-		if (!this.wasMessagePostedToSlack(this.botName.toLowerCase(), edgeState.conflictCl)) {
-			throw new Error(`no message sent for CL#${edgeState.conflictCl}`)
+
+		const conflictCl = edgeState.conflict && edgeState.conflict.change
+		if (!conflictCl) {
+			throw new Error('no conflict cl in edge state')
+		}
+
+		if (!this.wasMessagePostedToSlack(this.botName.toLowerCase(), conflictCl)) {
+			throw new Error(`no message sent for CL#${conflictCl}`)
 		}
 	}
 
 	async verifyAndPerformStomp(source: string, target: string, additionalSlackChannel?: string) {
 		const edgeState: EdgeState = await this.getEdgeState(source, target)
 
+		const conflictCl = edgeState.conflict && edgeState.conflict.change
+		if (!conflictCl) {
+			throw new Error('no conflict cl in edge state')
+		}
+
 		const checkSlack = async () => {
 			const slackChannel = this.botName.toLowerCase()
 			const channelsStr = slackChannel + (additionalSlackChannel ? ' and ' + additionalSlackChannel : '')
-			this.info(`Ensuring Slack message sent to ${channelsStr} for CL#${edgeState.conflictCl}`)
+
+			this.info(`Ensuring Slack message sent to ${channelsStr} for CL#${conflictCl}`)
 
 			const channels = [slackChannel]
 			if (additionalSlackChannel) {
@@ -676,13 +697,13 @@ export abstract class FunctionalTest {
 			}
 
 			for (const channel of channels) {
-				if (!await this.wasMessagePostedToSlack(channel, edgeState.conflictCl)) {
-					throw new Error(`No Slack message sent to ${channel} for CL#${edgeState.conflictCl}`)	
+				if (!await this.wasMessagePostedToSlack(channel, conflictCl)) {
+					throw new Error(`No Slack message sent to ${channel} for CL#${conflictCl}`)	
 				}
 			}
 		}
 
-		this.info(`Stomp ${source} -> ${target} @ CL#${edgeState.conflictCl}`)
+		this.info(`Stomp ${source} -> ${target} @ CL#${conflictCl}`)
 		const [verifyResult] = await Promise.all([
 			this.verifyStompRequest(source, target, edgeState),
 			checkSlack()
@@ -712,7 +733,7 @@ export abstract class FunctionalTest {
 			throw new Error('edge must be blocked to create shelf!')
 		}
 
-		const conflictCl = edgeState.conflictCl
+		const conflictCl = edgeState.conflict && edgeState.conflict.change
 
 		this.info(`Create shelf for ${source} -> ${target} @${conflictCl}`)
 
