[NO NON-PII COMMIT MESSAGE]

diff --git a/Engine/Extras/RoboMerge/v3/functional_tests/framework/src/BranchState.ts b/Engine/Extras/RoboMerge/v3/functional_tests/framework/src/BranchState.ts
--- a/Engine/Extras/RoboMerge/v3/functional_tests/framework/src/BranchState.ts
+++ b/Engine/Extras/RoboMerge/v3/functional_tests/framework/src/BranchState.ts
@@ -1,5 +1,9 @@
 // Copyright Epic Games, Inc. All Rights Reserved.
 
+export type Blockage = {
+	owner: string
+	change?: number
+}
 
 export class EdgeState {
 	constructor(private rawEdgeData: any) {
@@ -14,8 +18,11 @@ export class EdgeState {
 		return !!this.rawEdgeData.blockage
 	}
 
-	get conflictCl() {
-		return this.rawEdgeData.blockage ? this.rawEdgeData.blockage.change : -1
+	get conflict() {
+		return this.rawEdgeData.blockage ? {
+			owner: this.rawEdgeData.blockage.owner || '<no owner>',
+			change: this.rawEdgeData.blockage.change
+		} : null
 	}
 
 	getLastCL() {
diff --git a/Engine/Extras/RoboMerge/v3/functional_tests/framework/src/tests/approval.ts b/Engine/Extras/RoboMerge/v3/functional_tests/framework/src/tests/approval.ts
--- a/Engine/Extras/RoboMerge/v3/functional_tests/framework/src/tests/approval.ts
+++ b/Engine/Extras/RoboMerge/v3/functional_tests/framework/src/tests/approval.ts
@@ -31,7 +31,8 @@ export class Approval extends SimpleMainAndReleaseTestBase {
 		  { from: this.fullBranchName('Release'), to: this.fullBranchName('Main')
 		  , approval: {
 		      description: 'hey there!',
-		      group: 'approvers'
+		      channelName: 'approvers',
+		      channelId: 'abc123'
 		    }
 		  }
 		]
diff --git a/Engine/Extras/RoboMerge/v3/functional_tests/framework/src/tests/multiple-conflicts.ts b/Engine/Extras/RoboMerge/v3/functional_tests/framework/src/tests/multiple-conflicts.ts
--- a/Engine/Extras/RoboMerge/v3/functional_tests/framework/src/tests/multiple-conflicts.ts
+++ b/Engine/Extras/RoboMerge/v3/functional_tests/framework/src/tests/multiple-conflicts.ts
@@ -36,10 +36,14 @@ export class MultipleConflicts extends MultipleDevAndReleaseTestBase {
 			this.verifyAndPerformStomp('Main', 'Dev-Pootle', 'Pootle'),
 			async () => {
 				const edgeState: EdgeState = await this.getEdgeState('Main', 'Dev-Perkin')
+				const conflictCl = edgeState.conflict && edgeState.conflict.change
+				if (!conflictCl) {
+					throw new Error('no conflict cl in edge state')
+				}
 
 				const results = await Promise.all([
-					this.wasMessagePostedToSlack(this.botName.toLowerCase(), edgeState.conflictCl),
-					this.wasMessagePostedToSlack('Pootle', edgeState.conflictCl)
+					this.wasMessagePostedToSlack(this.botName.toLowerCase(), conflictCl),
+					this.wasMessagePostedToSlack('Pootle', conflictCl)
 				])
 				if (!results[0]) {
 					throw new Error('Expected message!') // todo: details
