[NO NON-PII COMMIT MESSAGE]

diff --git a/Engine/Extras/RoboMerge/v3/functional_tests/framework/src/framework.ts b/Engine/Extras/RoboMerge/v3/functional_tests/framework/src/framework.ts
--- a/Engine/Extras/RoboMerge/v3/functional_tests/framework/src/framework.ts
+++ b/Engine/Extras/RoboMerge/v3/functional_tests/framework/src/framework.ts
@@ -208,7 +208,7 @@ export type EdgeProperties = Partial<EdgeOptionFields> & {
 	to: string
 }
 
-const NUM_WAIT_INTERVALS = 20
+const NUM_WAIT_INTERVALS = 25
 export async function retryWithBackoff<T extends {}>(desc: string, f: (last: boolean) => Promise<T | null>): Promise<T> {
 	let sleepTime = .5
 	for (let safety = 0; safety < NUM_WAIT_INTERVALS; ++safety) {
