[NO NON-PII COMMIT MESSAGE]

diff --git a/Engine/Extras/RoboMerge/v3/functional_tests/framework/src/tests.ts b/Engine/Extras/RoboMerge/v3/functional_tests/framework/src/tests.ts
--- a/Engine/Extras/RoboMerge/v3/functional_tests/framework/src/tests.ts
+++ b/Engine/Extras/RoboMerge/v3/functional_tests/framework/src/tests.ts
@@ -279,8 +279,8 @@ async function go() {
 	///////////////////////
 	// TESTS TO RUN 
 
-	const specificTests = /*/[availableTests[17]]/*/availableTests  /*/ .slice(0, 0)/**/
-	const targetTests = /*/[availableTargetTests[7]] /*/ availableTargetTests /*/ .slice(8, 10) /**/
+	const specificTests = /*/[availableTests[17]]/*/availableTests  /* .slice(0, 0)/**/
+	const targetTests = /**/[availableTargetTests[5], availableTargetTests[8], availableTargetTests[9]] /*/ availableTargetTests /* .slice(8, 10) /**/
 	const tests = [...specificTests, ...targetTests]
 
 	//
