[NO NON-PII COMMIT MESSAGE]

diff --git a/Engine/Extras/RoboMerge/v3/src/robo/edgebot.ts b/Engine/Extras/RoboMerge/v3/src/robo/edgebot.ts
--- a/Engine/Extras/RoboMerge/v3/src/robo/edgebot.ts
+++ b/Engine/Extras/RoboMerge/v3/src/robo/edgebot.ts
@@ -11,7 +11,7 @@ import { IntegrationTarget, isExecP4Error, OpenedFileRecord, PerforceContext, EX
 import { VersionReader } from "../common/version";
 import { EdgeBotInterface, IPCControls, ReconsiderArgs } from "./bot-interfaces";
 import { PauseState } from "./state-interfaces";
-import { BlockagePauseInfo, BlockagePauseInfoMinimal } from "./status-types";
+import { BlockagePauseInfo, BlockagePauseInfoMinimal, EdgeStatusFields } from "./status-types";
 import { AlreadyIntegrated, Branch, ChangeInfo, ConflictingFile, Failure, MergeAction, PendingChange } from "./branch-interfaces";
 import { PersistentConflict } from "./conflict-interfaces";
 import { BotEventTriggers } from "./events";
@@ -905,7 +905,7 @@ class EdgeBotImpl extends PerforceStatefulBot {
 	}
 
 	createStatus() {
-		const status: { [key: string]: any } = {}
+		const status: Partial<EdgeStatusFields> = {}
 
 		status.name = this.fullName
 		status.display_name = this.displayName
@@ -946,11 +946,11 @@ class EdgeBotImpl extends PerforceStatefulBot {
 		if (this.targetBranch.resolver) {
 			status.resolver = this.targetBranch.resolver
 		}
-		status.skipAllowed = !this.options.disallowSkip
+		status.disallowSkip = this.options.disallowSkip
 		status.incognitoMode = this.options.incognitoMode
 		status.excludeAuthors = this.options.excludeAuthors
 
-		return status
+		return status as EdgeStatusFields
 	}
 
 	public forceSetLastClWithContext(value: number, culprit: string, reason: string, unblock?: boolean) {
diff --git a/Engine/Extras/RoboMerge/v3/src/robo/status-types.ts b/Engine/Extras/RoboMerge/v3/src/robo/status-types.ts
--- a/Engine/Extras/RoboMerge/v3/src/robo/status-types.ts
+++ b/Engine/Extras/RoboMerge/v3/src/robo/status-types.ts
@@ -117,6 +117,8 @@ export type EdgeStatusFields = BotStatusFields & {
 	lastGoodCL?: number
 	lastGoodCLJobLink?: string
 	lastGoodCLDate?: Date
+
+	num_changes_remaining: number
 }
 
 export type NodeStatusFields = BotStatusFields & {
