[NO NON-PII COMMIT MESSAGE]

diff --git a/Engine/Extras/RoboMerge/v3/run-tests.sh b/Engine/Extras/RoboMerge/v3/run-tests.sh
--- a/Engine/Extras/RoboMerge/v3/run-tests.sh
+++ b/Engine/Extras/RoboMerge/v3/run-tests.sh
@@ -36,6 +36,7 @@ cleanup
 
 docker network create --driver bridge robomerge_functtest_network
 
+
 docker build -t helix -f helix.Dockerfile .
 docker build -t helix-and-node -f helix-and-node.Dockerfile .
 
@@ -76,7 +77,7 @@ docker run -d -p 8877:8877 -p 8811:8811 -p 25:25 -h robomerge_functtest --name r
     -e ROBO_BRANCHSPECS_ROOT_PATH=//RoboMergeData/Main  \
     -e ROBO_NO_TLS=true  \
     -e ROBO_EXTERNAL_URL=http://robomerge_functtest:8877  \
-    -e BOTNAME=ft1,ft2,ft3,ft4  \
+    -e BOTNAME=ft1,ft2,ft3,ft4,targets  \
     -e ROBO_LOG_LEVEL=info \
     -e ROBO_SLACK_DOMAIN=http://localhost:8811 \
     robomerge node --trace-warnings dist/robo/watchdog.js
