Delay processing changes if the user is robomerge or the workspace starts with "swarm-" to give change owner commands time to complete processing so the author is not lost when robomerge integrates the change

diff --git a/Engine/Extras/RoboMerge/v3/src/robo/nodebot.ts b/Engine/Extras/RoboMerge/v3/src/robo/nodebot.ts
--- a/Engine/Extras/RoboMerge/v3/src/robo/nodebot.ts
+++ b/Engine/Extras/RoboMerge/v3/src/robo/nodebot.ts
@@ -1294,6 +1294,20 @@ export class NodeBot extends PerforceStatefulBot implements NodeBotInterface {
 
 		for (let changeIndex = 0; changeIndex < changes.length; ++changeIndex) {
 			const change = changes[changeIndex]
+
+			// If the user of a change is robomerge or this is a change submitted via swarm on behalf of someone
+			// don't process the change for a minute to give the change owner command time to be processed
+			if ((Date.now() - (change.time || 0)) < 60000) {
+				if (change.user === "robomerge") {
+					this.nodeBotLogger.info(`Delaying processing of ${change.change} while robomerge changes owner.`)
+					return
+				}
+				else if (change.client.startsWith("swarm-")) {
+					this.nodeBotLogger.info(`Delaying processing of ${change.change} to give swarm time to change owner.`)
+					return
+				}
+			}
+
 			const changeResult = await this._processAndMergeCl(availableEdges, change, false)
 
 			// Exit immediately on syntax errors
@@ -1742,7 +1756,7 @@ export class NodeBot extends PerforceStatefulBot implements NodeBotInterface {
 
 			const targetEdge = availableEdges.get(target.branch.upperName)
 
-			// If the edge requested is currently no available this tick, simply log that we skipped it and move on
+			// If the edge requested is currently not available this tick, simply log that we skipped it and move on
 			if (!targetEdge) {
 				this.nodeBotLogger.verbose(`No available edges can serve CL ${info.cl} -> ${target.branch.name} this tick.`)
 				continue
