Put in robomerge user debugging info print outs

diff --git a/Engine/Extras/RoboMerge/v3/src/robo/nodebot.ts b/Engine/Extras/RoboMerge/v3/src/robo/nodebot.ts
--- a/Engine/Extras/RoboMerge/v3/src/robo/nodebot.ts
+++ b/Engine/Extras/RoboMerge/v3/src/robo/nodebot.ts
@@ -1308,6 +1308,10 @@ export class NodeBot extends PerforceStatefulBot implements NodeBotInterface {
 				}
 			}
 
+			if (change.user === "robomerge") {
+				this.nodeBotLogger.info(`Processing change ${change.change} by user robomerge. Now: ${Date.now()} Change time: ${change.time}`)
+            }
+
 			const changeResult = await this._processAndMergeCl(availableEdges, change, false)
 
 			// Exit immediately on syntax errors
