[NO NON-PII COMMIT MESSAGE]

diff --git a/Engine/Extras/RoboMerge/v3/src/robo/nodebot.ts b/Engine/Extras/RoboMerge/v3/src/robo/nodebot.ts
--- a/Engine/Extras/RoboMerge/v3/src/robo/nodebot.ts
+++ b/Engine/Extras/RoboMerge/v3/src/robo/nodebot.ts
@@ -1294,10 +1294,10 @@ export class NodeBot extends PerforceStatefulBot implements NodeBotInterface {
 
 		for (let changeIndex = 0; changeIndex < changes.length; ++changeIndex) {
 			const change = changes[changeIndex]
-
+			
 			// If the user of a change is robomerge or this is a change submitted via swarm on behalf of someone
 			// don't process the change for a minute to give the change owner command time to be processed
-			if ((Date.now() - (change.time || 0)) < 60000) {
+			if ((Date.now() / 1000 - (change.time || 0)) < 60000) {
 				if (change.user === "robomerge") {
 					this.nodeBotLogger.info(`Delaying processing of ${change.change} while robomerge changes owner.`)
 					return
