Add additional logging when delaying a change

diff --git a/Engine/Extras/RoboMerge/v3/src/robo/nodebot.ts b/Engine/Extras/RoboMerge/v3/src/robo/nodebot.ts
--- a/Engine/Extras/RoboMerge/v3/src/robo/nodebot.ts
+++ b/Engine/Extras/RoboMerge/v3/src/robo/nodebot.ts
@@ -1294,16 +1294,16 @@ export class NodeBot extends PerforceStatefulBot implements NodeBotInterface {
 
 		for (let changeIndex = 0; changeIndex < changes.length; ++changeIndex) {
 			const change = changes[changeIndex]
-			
+
 			// If the user of a change is robomerge or this is a change submitted via swarm on behalf of someone
 			// don't process the change for a minute to give the change owner command time to be processed
-			if ((Date.now() / 1000 - (change.time || 0)) < 60000) {
+			if ((Date.now() - (change.time || 0)) < 60000) {
 				if (change.user === "robomerge") {
-					this.nodeBotLogger.info(`Delaying processing of ${change.change} while robomerge changes owner.`)
+					this.nodeBotLogger.info(`Delaying processing of ${change.change} while robomerge changes owner.  Now: ${Date.now()} Change time: ${change.time}`)
 					return
 				}
 				else if (change.client.startsWith("swarm-")) {
-					this.nodeBotLogger.info(`Delaying processing of ${change.change} to give swarm time to change owner.`)
+					this.nodeBotLogger.info(`Delaying processing of ${change.change} to give swarm time to change owner.  Now: ${Date.now()} Change time: ${change.time}`)
 					return
 				}
 			}
