Put time comparison back in to ms so that comparison against 60000 isn't waiting 16 2/3 hours
Make processing a robomerge user message a warning

diff --git a/Engine/Extras/RoboMerge/v3/src/robo/nodebot.ts b/Engine/Extras/RoboMerge/v3/src/robo/nodebot.ts
--- a/Engine/Extras/RoboMerge/v3/src/robo/nodebot.ts
+++ b/Engine/Extras/RoboMerge/v3/src/robo/nodebot.ts
@@ -1297,7 +1297,7 @@ export class NodeBot extends PerforceStatefulBot implements NodeBotInterface {
 
 			// If the user of a change is robomerge or this is a change submitted via swarm on behalf of someone
 			// don't process the change for a minute to give the change owner command time to be processed
-			if ((Date.now() / 1000 - (change.time || 0)) < 60000) {
+			if ((Date.now() - ((change.time || 0) * 1000)) < 60000) {
 				if (change.user === "robomerge") {
 					this.nodeBotLogger.info(`Delaying processing of ${change.change} while robomerge changes owner.  Now: ${Date.now()} Change time: ${change.time}`)
 					return
@@ -1309,7 +1309,7 @@ export class NodeBot extends PerforceStatefulBot implements NodeBotInterface {
 			}
 
 			if (change.user === "robomerge") {
-				this.nodeBotLogger.info(`Processing change ${change.change} by user robomerge. Now: ${Date.now()} Change time: ${change.time}`)
+				this.nodeBotLogger.warn(`Processing change ${change.change} by user robomerge. Now: ${Date.now()} Change time: ${change.time}`)
             }
 
 			const changeResult = await this._processAndMergeCl(availableEdges, change, false)
