[NO NON-PII COMMIT MESSAGE]

diff --git a/Engine/Extras/RoboMerge/v3/functional_tests/framework/src/framework.ts b/Engine/Extras/RoboMerge/v3/functional_tests/framework/src/framework.ts
--- a/Engine/Extras/RoboMerge/v3/functional_tests/framework/src/framework.ts
+++ b/Engine/Extras/RoboMerge/v3/functional_tests/framework/src/framework.ts
@@ -50,7 +50,8 @@ process.on('unhandledRejection', async (err) => {
 })
 
 export class P4Client {
-	constructor(private p4: Perforce, public user: string, public workspace: string, public root: string, public stream: string, public name: string) {
+	constructor(private p4: Perforce, public user: string, public workspace: string,
+		public root: string, public stream: string, public name: string) {
 	}
 
 	async create(spec: string) {
@@ -322,7 +323,7 @@ export abstract class FunctionalTest {
 	}
 
 
-	protected makeForceAllBranchDef(stream: string, to: string[]/*, additionalProps?: {[prop: string]: string}[]*/): RobomergeBranchSpec {
+	protected makeForceAllBranchDef(stream: string, to: string[]): RobomergeBranchSpec {
 		return this.makeBranchDef(stream, to, true)
 	}
 
@@ -657,30 +658,58 @@ export abstract class FunctionalTest {
 		}
 	}
 
-	wasMessagePostedToSlack(channel: string, cl: number) {
+	wasMessagePostedToSlack(channel: string, cl: number, target: string) {
+		const targetName = this.fullBranchName(target)
 		return getJson(DUMMY_SLACK_DOMAIN + '/posted/' + channel)
-			.then((cls: number[]) => cls.indexOf(cl) >= 0)
+			.then((perTargetCls: [string, number[]][]) => {
+				for (const [target, cls] of perTargetCls) {
+					if (target === targetName) {
+						return cls.indexOf(cl) >= 0
+					}
+				}
+				return false
+			})
 	}
 
-	async ensureConflictMessagePostedToSlack(sourceStream: string, targetStream: string) {
-		const edgeDisplayName = `${sourceStream} -> ${targetStream}`
-		this.info(`Ensuring conflict message sent for ${edgeDisplayName}`)
-		const edgeState = await this.getEdgeState('Release', 'Main')
+	async wasConflictMessagePostedToSlack(sourceStream: string, targetStream: string, optChannel?: string)
+		: Promise<[boolean, number]>
+	{
+		const edgeState = await this.getEdgeState(sourceStream, targetStream)
 
 		const conflictCl = edgeState.conflict && edgeState.conflict.change
 		if (!conflictCl) {
 			throw new Error('no conflict cl in edge state')
 		}
 
-		if (!this.wasMessagePostedToSlack(this.botName.toLowerCase(), conflictCl)) {
-			throw new Error(`no message sent for CL#${conflictCl}`)
+		return [await this.wasMessagePostedToSlack(optChannel || this.botName.toLowerCase(), conflictCl, targetStream), conflictCl]
+	}
+
+	async ensureConflictMessagePostedToSlack(sourceStream: string, targetStream: string, optChannel?: string) {
+		const edgeDisplayName = `${sourceStream} -> ${targetStream}`
+		this.info(`Ensuring conflict message sent for ${edgeDisplayName}`)
+
+		const [messageSent, conflictCl] = await this.wasConflictMessagePostedToSlack(sourceStream, targetStream, optChannel)
+		if (!messageSent) {
+			const channelMessage = optChannel ? `to '${optChannel}'' ` : ''
+			throw new Error(`no message sent ${channelMessage}for CL#${conflictCl} (${edgeDisplayName})`)
+		}
+	}
+
+	async ensureNoConflictMessagePostedToSlack(sourceStream: string, targetStream: string, optChannel?: string) {
+		const edgeDisplayName = `${sourceStream} -> ${targetStream}`
+		this.info(`Ensuring no conflict message sent for ${edgeDisplayName}`)
+
+		const [messageSent, conflictCl] = await this.wasConflictMessagePostedToSlack(sourceStream, targetStream, optChannel)
+		if (messageSent) {
+			const channelMessage = optChannel ? `to '${optChannel}'' ` : ''
+			throw new Error(`unexpected message sent ${channelMessage}for CL#${conflictCl} (${edgeDisplayName})`)
 		}
 	}
 
 	async verifyAndPerformStomp(source: string, target: string, additionalSlackChannel?: string) {
 		const edgeState: EdgeState = await this.getEdgeState(source, target)
 
-		const conflictCl = edgeState.conflict && edgeState.conflict.change
+		const conflictCl: number | undefined = edgeState.conflict && edgeState.conflict.change
 		if (!conflictCl) {
 			throw new Error('no conflict cl in edge state')
 		}
@@ -696,11 +725,8 @@ export abstract class FunctionalTest {
 				channels.push(additionalSlackChannel)
 			}
 
-			for (const channel of channels) {
-				if (!await this.wasMessagePostedToSlack(channel, conflictCl)) {
-					throw new Error(`No Slack message sent to ${channel} for CL#${conflictCl}`)	
-				}
-			}
+			await Promise.all(channels.map(channel =>
+				this.ensureConflictMessagePostedToSlack(source, target, channel)))
 		}
 
 		this.info(`Stomp ${source} -> ${target} @ CL#${conflictCl}`)
@@ -714,7 +740,8 @@ export abstract class FunctionalTest {
 			this.warn(verifyResult.message)
 			// this.warn("nonBinaryFilesResolved=" + verifyResult.nonBinaryFilesResolved)
 			// this.warn("remainingAllBinary=" + verifyResult.remainingAllBinary)
-			// this.warn("files=" + (Array.isArray(verifyResult.files) ? verifyResult.files[0].targetFileName : `no files (${verifyResult.files})`))
+			// this.warn("files=" + (Array.isArray(verifyResult.files)
+				// ? verifyResult.files[0].targetFileName : `no files (${verifyResult.files})`))
 			throw new Error('Stomp verify returned unexpected values')
 		}
 
@@ -745,7 +772,8 @@ export abstract class FunctionalTest {
 			.replace('<bot>', this.botName)
 			.replace('<node>', sourceBranchName)
 			.replace('<op>', 'create_shelf')
-		const response = await bent('POST', 'string')(`${endpoint}?cl=${conflictCl}&target=${targetBranchName}&workspace=${targetClient.workspace}`)
+		const response = await bent('POST', 'string')
+			(`${endpoint}?cl=${conflictCl}&target=${targetBranchName}&workspace=${targetClient.workspace}`)
 
 		if (response !== 'ok') {
 			throw new Error('unexpected response from create_shelf: ' + response)
@@ -975,5 +1003,4 @@ export class P4Util {
 	static escapeBranchName(stream: string) {
 		return stream.replace(/[^A-Za-z0-9]/g, '_')
 	}
-	
-}
\ No newline at end of file
+}
diff --git a/Engine/Extras/RoboMerge/v3/functional_tests/framework/src/tests/multiple-conflicts.ts b/Engine/Extras/RoboMerge/v3/functional_tests/framework/src/tests/multiple-conflicts.ts
--- a/Engine/Extras/RoboMerge/v3/functional_tests/framework/src/tests/multiple-conflicts.ts
+++ b/Engine/Extras/RoboMerge/v3/functional_tests/framework/src/tests/multiple-conflicts.ts
@@ -1,6 +1,5 @@
 // Copyright Epic Games, Inc. All Rights Reserved.
 import { P4Util } from '../framework'
-import { EdgeState } from '../BranchState'
 import { MultipleDevAndReleaseTestBase } from '../MultipleDevAndReleaseTestBase'
 
 
@@ -34,28 +33,11 @@ export class MultipleConflicts extends MultipleDevAndReleaseTestBase {
 			this.checkHeadRevision('Dev-Perkin', 'test.uasset', 2),
 			this.checkHeadRevision('Dev-Pootle', 'test.uasset', 2),
 			this.verifyAndPerformStomp('Main', 'Dev-Pootle', 'Pootle'),
-			async () => {
-				const edgeState: EdgeState = await this.getEdgeState('Main', 'Dev-Perkin')
-				const conflictCl = edgeState.conflict && edgeState.conflict.change
-				if (!conflictCl) {
-					throw new Error('no conflict cl in edge state')
-				}
-
-				const results = await Promise.all([
-					this.wasMessagePostedToSlack(this.botName.toLowerCase(), conflictCl),
-					this.wasMessagePostedToSlack('Pootle', conflictCl)
-				])
-				if (!results[0]) {
-					throw new Error('Expected message!') // todo: details
-				}
-
-				if (results[1]) {
-					throw new Error('Not expecting message!') // todo: details
-				}
-			}
+			this.ensureConflictMessagePostedToSlack('Main', 'Dev-Perkin'),
+			this.ensureNoConflictMessagePostedToSlack('Main', 'Dev-Perkin', 'Pootle')
 		])
 
-		// Give RM plenty of time to stomp and merge all those test.uasset revisions 
+		// Wait for RM to stomp and merge all those test.uasset revisions 
 		await this.waitForRobomergeIdle()
 
 		await Promise.all([
diff --git a/Engine/Extras/RoboMerge/v3/src/common/dummyslackserver.ts b/Engine/Extras/RoboMerge/v3/src/common/dummyslackserver.ts
--- a/Engine/Extras/RoboMerge/v3/src/common/dummyslackserver.ts
+++ b/Engine/Extras/RoboMerge/v3/src/common/dummyslackserver.ts
@@ -3,7 +3,8 @@
 import { AppInterface, Handler, WebRequest } from './webserver'
 import { setDefault } from './helper'
 
-const posted = new Map<string, number[]>()
+// channel -> target -> changes
+const posted = new Map<string, Map<string, number[]>>()
 
 
 export class DummySlackApp implements AppInterface {
@@ -14,17 +15,30 @@ export class DummySlackApp implements AppInterface {
 	@Handler('POST', '/api/*')
 	post(command: string) {
 		// @todo (maybe) construct app with logger
-		console.log('POST', command, this.req.reqData)
 
 		const data = this.req.reqData
 		if (!data) {
 			throw new Error('Nothing to post!')
 		}
 		const clMatch = data.match(/changes\/(\d+)\|\d+/)
-		if (clMatch) {
-			const channel = JSON.parse(data).channel
-			console.log(JSON.parse(data).channel, clMatch[1])
-			setDefault(posted, channel, []).push(parseInt(clMatch[1]))
+		const edgeMatch = data.match(/"(\w+) -> (\w+)"/)
+		if (clMatch && edgeMatch) {
+			let dataObj: any
+			try {
+				dataObj = JSON.parse(data)
+			}
+			catch (exc) {
+				console.log('Non-JSON sent to dummy-slack: ', data)
+				return
+			}
+
+			setDefault(
+				setDefault(posted, dataObj.channel, new Map),
+				edgeMatch[2], []
+			).push(parseInt(clMatch[1]))
+		}
+		else {
+			console.log('dummy-slack POST: cl or edge not found', command, this.req.reqData)
 		}
 
 		// look for color for block or resolve?
@@ -34,7 +48,6 @@ export class DummySlackApp implements AppInterface {
 
 	@Handler('GET', '/posted/*')
 	channelMessages(channel: string) {
-		console.log('GET POSTED!!!', channel, posted.size, posted.get(channel))
-		return posted.get(channel) || []
+		return [...(posted.get(channel) || [])]
 	}
 }
