[NO NON-PII COMMIT MESSAGE]

diff --git a/Engine/Extras/RoboMerge/v3/src/robo/notifications.ts b/Engine/Extras/RoboMerge/v3/src/robo/notifications.ts
--- a/Engine/Extras/RoboMerge/v3/src/robo/notifications.ts
+++ b/Engine/Extras/RoboMerge/v3/src/robo/notifications.ts
@@ -254,7 +254,7 @@ class SlackMessages {
 				timestamp = await this.slack.postMessage(message)
 			}
 			catch (err) {
-				console.error('Error talking to Slack! ' + err.toString())
+				this.smLogger.printException(err, 'Error talking to Slack')
 				return
 			}
 
