[NO NON-PII COMMIT MESSAGE]

diff --git a/Engine/Extras/RoboMerge/v3/functional_tests/framework/src/framework.ts b/Engine/Extras/RoboMerge/v3/functional_tests/framework/src/framework.ts
--- a/Engine/Extras/RoboMerge/v3/functional_tests/framework/src/framework.ts
+++ b/Engine/Extras/RoboMerge/v3/functional_tests/framework/src/framework.ts
@@ -690,7 +690,7 @@ export abstract class FunctionalTest {
 
 		const [messageSent, conflictCl] = await this.wasConflictMessagePostedToSlack(sourceStream, targetStream, optChannel)
 		if (!messageSent) {
-			const channelMessage = optChannel ? `to '${optChannel}'' ` : ''
+			const channelMessage = optChannel ? `to '${optChannel}' ` : ''
 			throw new Error(`no message sent ${channelMessage}for CL#${conflictCl} (${edgeDisplayName})`)
 		}
 	}
@@ -701,7 +701,7 @@ export abstract class FunctionalTest {
 
 		const [messageSent, conflictCl] = await this.wasConflictMessagePostedToSlack(sourceStream, targetStream, optChannel)
 		if (messageSent) {
-			const channelMessage = optChannel ? `to '${optChannel}'' ` : ''
+			const channelMessage = optChannel ? `to '${optChannel} ` : ''
 			throw new Error(`unexpected message sent ${channelMessage}for CL#${conflictCl} (${edgeDisplayName})`)
 		}
 	}
