[NO NON-PII COMMIT MESSAGE]

diff --git a/Engine/Extras/RoboMerge/v3/functional_tests/framework/src/tests.ts b/Engine/Extras/RoboMerge/v3/functional_tests/framework/src/tests.ts
--- a/Engine/Extras/RoboMerge/v3/functional_tests/framework/src/tests.ts
+++ b/Engine/Extras/RoboMerge/v3/functional_tests/framework/src/tests.ts
@@ -279,14 +279,17 @@ async function go() {
 	///////////////////////
 	// TESTS TO RUN 
 
-	const specificTests = /*/[availableTests[17]]/*/availableTests  /* .slice(0, 0)/**/
-	const targetTests = /*/[availableTargetTests[7]] /*/ availableTargetTests /* .slice(8, 10) /**/
-	const tests = [...specificTests, ...targetTests]
+	const specificTests = /**/[availableTests[17]]/*/availableTests  /* .slice(0, 0)/**/
+	// needs at least 1 target test otherwise it hangs
+	const targetTests = /**/[availableTargetTests[7]] /*/ availableTargetTests /* .slice(8, 10) /**/
 
 	//
 	///////////////////////
 
-	console.log(`${tests[0].testName} to ${tests[tests.length - 1].testName}`)
+	console.log(`${specificTests[0].testName} to ${specificTests[specificTests.length - 1].testName}`)
+	console.log(`${targetTests[0].testName} to ${targetTests[targetTests.length - 1].testName}`)
+
+	const tests = /**/ [...specificTests, ...targetTests] /*/ [...specificTests, ...targetTests] /**/
 
 	console.log(`Running set-up for ${tests.length} tests`)
 	await Promise.all(tests.map(test => test.setup()))
