[NO NON-PII COMMIT MESSAGE]

diff --git a/Engine/Extras/RoboMerge/v3/functional_tests/framework/src/framework.ts b/Engine/Extras/RoboMerge/v3/functional_tests/framework/src/framework.ts
--- a/Engine/Extras/RoboMerge/v3/functional_tests/framework/src/framework.ts
+++ b/Engine/Extras/RoboMerge/v3/functional_tests/framework/src/framework.ts
@@ -686,22 +686,22 @@ export abstract class FunctionalTest {
 
 	async ensureConflictMessagePostedToSlack(sourceStream: string, targetStream: string, optChannel?: string) {
 		const edgeDisplayName = `${sourceStream} -> ${targetStream}`
-		this.info(`Ensuring conflict message sent for ${edgeDisplayName}`)
+		const channelMessage = optChannel ? `to '${optChannel}' ` : ''
+		this.info(`Ensuring conflict message sent ${channelMessage}for ${edgeDisplayName}`)
 
 		const [messageSent, conflictCl] = await this.wasConflictMessagePostedToSlack(sourceStream, targetStream, optChannel)
 		if (!messageSent) {
-			const channelMessage = optChannel ? `to '${optChannel}' ` : ''
 			throw new Error(`no message sent ${channelMessage}for CL#${conflictCl} (${edgeDisplayName})`)
 		}
 	}
 
 	async ensureNoConflictMessagePostedToSlack(sourceStream: string, targetStream: string, optChannel?: string) {
 		const edgeDisplayName = `${sourceStream} -> ${targetStream}`
-		this.info(`Ensuring no conflict message sent for ${edgeDisplayName}`)
+		const channelMessage = optChannel ? `to '${optChannel}' ` : ''
+		this.info(`Ensuring no conflict message sent ${channelMessage}for ${edgeDisplayName}`)
 
 		const [messageSent, conflictCl] = await this.wasConflictMessagePostedToSlack(sourceStream, targetStream, optChannel)
 		if (messageSent) {
-			const channelMessage = optChannel ? `to '${optChannel} ` : ''
 			throw new Error(`unexpected message sent ${channelMessage}for CL#${conflictCl} (${edgeDisplayName})`)
 		}
 	}
diff --git a/Engine/Extras/RoboMerge/v3/functional_tests/framework/src/tests.ts b/Engine/Extras/RoboMerge/v3/functional_tests/framework/src/tests.ts
--- a/Engine/Extras/RoboMerge/v3/functional_tests/framework/src/tests.ts
+++ b/Engine/Extras/RoboMerge/v3/functional_tests/framework/src/tests.ts
@@ -279,17 +279,18 @@ async function go() {
 	///////////////////////
 	// TESTS TO RUN 
 
-	const specificTests = /**/[availableTests[17]]/*/availableTests  /* .slice(0, 0)/**/
-	// needs at least 1 target test otherwise it hangs
-	const targetTests = /**/[availableTargetTests[7]] /*/ availableTargetTests /* .slice(8, 10) /**/
+	const specificTests = /**/[availableTests[12]]/*/availableTests  /* .slice(0, 0)/**/
+	const targetTests: FunctionalTest[] = /**/ [] /*/ availableTargetTests  /**/
 
 	//
 	///////////////////////
 
 	console.log(`${specificTests[0].testName} to ${specificTests[specificTests.length - 1].testName}`)
-	console.log(`${targetTests[0].testName} to ${targetTests[targetTests.length - 1].testName}`)
+	if (targetTests.length > 0) {
+		console.log(`${targetTests[0].testName} to ${targetTests[targetTests.length - 1].testName}`)
+	}
 
-	const tests = /**/ [...specificTests, ...targetTests] /*/ [...specificTests, ...targetTests] /**/
+	const tests = [...specificTests, ...targetTests]
 
 	console.log(`Running set-up for ${tests.length} tests`)
 	await Promise.all(tests.map(test => test.setup()))
