Add ability to only post to the additional slack channel

diff --git a/Engine/Extras/RoboMerge/v3/functional_tests/framework/src/tests.ts b/Engine/Extras/RoboMerge/v3/functional_tests/framework/src/tests.ts
--- a/Engine/Extras/RoboMerge/v3/functional_tests/framework/src/tests.ts
+++ b/Engine/Extras/RoboMerge/v3/functional_tests/framework/src/tests.ts
@@ -26,6 +26,7 @@ import { MergeMainRevToRelease } from './tests/merge-main-rev-to-release'
 import { MultipleConflicts } from './tests/multiple-conflicts'
 import { MultipleRoutesToSkip } from './tests/multiple-routes-to-skip'
 import { OverriddenCommand } from './tests/overridden-command'
+import { PostToAdditionalChannel } from './tests/post-to-additional-channel'
 import { RejectBranchResolveStomp } from './tests/reject-branch-resolve-stomp'
 import { RejectDeleteResolveStomp } from './tests/reject-delete-resolve-stomp'
 import { RejectTextConflictStomp } from './tests/reject-text-conflict-stomp'
@@ -213,6 +214,8 @@ async function go() {
 		new StompForwardingCommands(p4),
 		new MultipleRoutesToSkip(p4),
 		new UnreachableSkip(p4), // 45
+
+		new PostToAdditionalChannel(p4),
 	]
 
 	const TARGET_TEST_DEFS: [string[], string | null][] = [
@@ -279,8 +282,14 @@ async function go() {
 	///////////////////////
 	// TESTS TO RUN 
 
-	const specificTests = /**/[availableTests[12]]/*/availableTests  /* .slice(0, 0)/**/
-	const targetTests: FunctionalTest[] = /**/ [] /*/ availableTargetTests  /**/
+	const specificTests = availableTests
+	// const specificTests = [availableTests[12]]
+	// const specificTests = availableTests.slice(0, 0)
+
+	const targetTests = availableTargetTests
+	// const targetTests: FunctionalTest[] = []
+	// const targetTests = [availableTargetTests[7]]
+	// const targetTests = availableTargetTests.slice(0, 0)
 
 	//
 	///////////////////////
diff --git a/Engine/Extras/RoboMerge/v3/functional_tests/framework/src/tests/post-to-additional-channel.ts b/Engine/Extras/RoboMerge/v3/functional_tests/framework/src/tests/post-to-additional-channel.ts
new file mode 100644
--- /dev/null
+++ b/Engine/Extras/RoboMerge/v3/functional_tests/framework/src/tests/post-to-additional-channel.ts
@@ -0,0 +1,66 @@
+// Copyright Epic Games, Inc. All Rights Reserved.
+import { P4Util } from '../framework'
+import { MultipleDevAndReleaseTestBase } from '../MultipleDevAndReleaseTestBase'
+
+
+export class PostToAdditionalChannel extends MultipleDevAndReleaseTestBase {
+	async setup() {
+		await this.p4.depot('stream', this.depotSpec())
+		await this.createStreamsAndWorkspaces()
+
+		await P4Util.addFileAndSubmit(this.getClient('Main'), 'test.uasset', 'Initial content', true)
+		await P4Util.addFileAndSubmit(this.getClient('Main'), 'test2.uasset', 'Initial content', true)
+
+		await this.initialPopulate()
+	}
+
+	async run() {
+		// prepare some conflicts
+		await Promise.all(['Dev-Perkin']
+			.map(branchName => this.getClient(branchName))
+			.map(client => client.sync()
+				.then(() => P4Util.editFileAndSubmit(client, 'test.uasset', 'Different content')))
+		)
+
+		// prepare some conflicts
+		await Promise.all(['Dev-Pootle']
+			.map(branchName => this.getClient(branchName))
+			.map(client => client.sync()
+				.then(() => P4Util.editFileAndSubmit(client, 'test2.uasset', 'Different content')))
+		)
+
+		const r1client = this.getClient('Release-1.0')
+		await r1client.sync()
+		await P4Util.editFileAndSubmit(r1client, 'test.uasset', 'Initial content\nSome more!')
+		await P4Util.editFileAndSubmit(r1client, 'test2.uasset', 'Initial content\nSome more!')
+	}
+
+	async verify() {
+
+		this.info(`Ensuring messages written to expected channels`)
+
+		await Promise.all([
+			this.ensureConflictMessagePostedToSlack('Main', 'Dev-Perkin'),
+			this.ensureConflictMessagePostedToSlack('Main', 'Dev-Perkin', 'Perkin'),
+			this.ensureNoConflictMessagePostedToSlack('Main', 'Dev-Perkin', 'Pootle'),
+			this.ensureNoConflictMessagePostedToSlack('Main', 'Dev-Pootle'),
+			this.ensureNoConflictMessagePostedToSlack('Main', 'Dev-Pootle', 'Perkin'),
+			this.ensureConflictMessagePostedToSlack('Main', 'Dev-Pootle', 'Pootle')
+		])
+	}
+
+	getEdges() {
+		return [{
+			from: this.fullBranchName('Main'),
+			to: this.fullBranchName('Dev-Perkin'),
+			additionalSlackChannel: 'Perkin'
+		},
+		{
+			from: this.fullBranchName('Main'),
+			to: this.fullBranchName('Dev-Pootle'),
+			additionalSlackChannel: 'Pootle',
+			postOnlyToAdditionalChannel: true
+		}]
+	}
+
+}
diff --git a/Engine/Extras/RoboMerge/v3/src/robo/branchdefs.ts b/Engine/Extras/RoboMerge/v3/src/robo/branchdefs.ts
--- a/Engine/Extras/RoboMerge/v3/src/robo/branchdefs.ts
+++ b/Engine/Extras/RoboMerge/v3/src/robo/branchdefs.ts
@@ -117,6 +117,7 @@ const nodeOptionFieldsPrototype = {
 	// if set, still generate workspace but use this name
 	workspaceNameOverride: '',
 	additionalSlackChannelForBlockages: '',
+	postMessagesToAdditionalChannelOnly: false,
 	ignoreBranchspecs: false,
 
 	badgeUrlOverride: '',
@@ -133,6 +134,8 @@ const edgeOptionFieldsPrototype = {
 	branchspec: '',
 	additionalSlackChannel: '',
 
+	postOnlyToAdditionalChannel: false,
+
 	terminal: false, // changes go along terminal edges but no further
 
 	implicitCommands: [''],
diff --git a/Engine/Extras/RoboMerge/v3/src/robo/branchgraph.ts b/Engine/Extras/RoboMerge/v3/src/robo/branchgraph.ts
--- a/Engine/Extras/RoboMerge/v3/src/robo/branchgraph.ts
+++ b/Engine/Extras/RoboMerge/v3/src/robo/branchgraph.ts
@@ -83,6 +83,7 @@ if (botname === '__TEST__') {
 		// also flag for applying reverse entry too? would apply to branchspecs, not so simple for resolver
 		const propertiesToPropagate: [keyof NodeOptions, keyof EdgeOptions, ConfigBlendMode | null][] = [
 			['additionalSlackChannelForBlockages', 'additionalSlackChannel', null],
+			['postMessagesToAdditionalChannelOnly', 'postOnlyToAdditionalChannel', null],
 			['lastGoodCLPath', 'lastGoodCLPath', null],
 			['disallowSkip', 'disallowSkip', null],
 			['incognitoMode', 'incognitoMode', null],
diff --git a/Engine/Extras/RoboMerge/v3/src/robo/edgebot.ts b/Engine/Extras/RoboMerge/v3/src/robo/edgebot.ts
--- a/Engine/Extras/RoboMerge/v3/src/robo/edgebot.ts
+++ b/Engine/Extras/RoboMerge/v3/src/robo/edgebot.ts
@@ -744,7 +744,9 @@ class EdgeBotImpl extends PerforceStatefulBot {
 						let msg = `Couldn't change changelist owner for change ${finalCl}: Perforce reports user ${userNonexistenceMatches[0]} doesn't exist`
 						this.edgeBotLogger.warn(msg + `:\n${output}`)
 						msg += `. User will *remain* as Robomerge:\n\`\`\`${output}\`\`\``
-						postMessageToChannel(msg, this.sourceNode.branchGraph.config.slackChannel, SlackMessageStyles.WARNING)
+						if (!this.targetBranch.config.postMessagesToAdditionalChannelOnly) {
+							postMessageToChannel(msg, this.sourceNode.branchGraph.config.slackChannel, SlackMessageStyles.WARNING)
+						}
 						if (this.targetBranch.config.additionalSlackChannelForBlockages) {
 							postMessageToChannel(msg, this.targetBranch.config.additionalSlackChannelForBlockages, SlackMessageStyles.WARNING)
 						}
diff --git a/Engine/Extras/RoboMerge/v3/src/robo/graphbot.ts b/Engine/Extras/RoboMerge/v3/src/robo/graphbot.ts
--- a/Engine/Extras/RoboMerge/v3/src/robo/graphbot.ts
+++ b/Engine/Extras/RoboMerge/v3/src/robo/graphbot.ts
@@ -112,12 +112,12 @@ export class GraphBot implements GraphInterface, BotEventHandler {
 		// doing it here ensures that we're using the most up-to-date config, e.g. after a branch spec reload
 
 		// preprocess to get a list of all additional Slack channels
-		const slackChannelOverrides: [Branch, Branch, string][] = []
+		const slackChannelOverrides: [Branch, Branch, string, boolean][] = []
 		for (const branch of this.branchGraph.branches) {
 			for (const [targetBranchName, edgeProps] of branch.edgeProperties) {
 				const slackChannel = edgeProps.additionalSlackChannel
 				if (slackChannel) {
-					slackChannelOverrides.push([branch, this.branchGraph.getBranch(targetBranchName)!, slackChannel])
+					slackChannelOverrides.push([branch, this.branchGraph.getBranch(targetBranchName)!, slackChannel, (edgeProps.postOnlyToAdditionalChannel || false)])
 				}
 			}
 		}
diff --git a/Engine/Extras/RoboMerge/v3/src/robo/notifications.ts b/Engine/Extras/RoboMerge/v3/src/robo/notifications.ts
--- a/Engine/Extras/RoboMerge/v3/src/robo/notifications.ts
+++ b/Engine/Extras/RoboMerge/v3/src/robo/notifications.ts
@@ -375,7 +375,7 @@ export class BotNotifications implements BotEventHandler {
 
 	constructor(private botname: string, slackChannel: string, persistence: Context, externalUrl: string, 
 		blockageUrlGenerator: NodeOpUrlGenerator, parentLogger: ContextualLogger,
-		slackChannelOverrides?: [Branch, Branch, string][]) {
+		slackChannelOverrides?: [Branch, Branch, string, boolean][]) {
 		this.botNotificationsLogger = parentLogger.createChild('Notifications')
 		// Hacky way to dynamically change the URL for notifications
 		this.externalRobomergeUrl = externalUrl
@@ -389,10 +389,10 @@ export class BotNotifications implements BotEventHandler {
 			if (slackChannelOverrides) {
 				this.additionalBlockChannelIds = new Map(slackChannelOverrides
 					.filter(
-						([_1, _2, channel]) => channel !== slackChannel
+						([_1, _2, channel, _3]) => channel !== slackChannel
 					)
 					.map(
-						([source, target, channel]) => [`${source.upperName}|${target.upperName}`, channel]
+						([source, target, channel, postOnlyToChannel]) => [`${source.upperName}|${target.upperName}`, [channel, postOnlyToChannel]]
 					)
 				)
 			}
@@ -453,18 +453,22 @@ export class BotNotifications implements BotEventHandler {
 			message.footer = blockage.failure.summary
 		}
 
-		// separate message for syntax errors (other blockages have a target branch)
-		const targetKey = targetBranch ? targetBranch.name : blockage.failure.kind
-
-		// Post message to channel
-		this.slackMessages.postOrUpdate(changeInfo.source_cl, targetKey, message)
-
-		// Post to additional side channel if present
-		if (targetBranch) { // temp: not for syntax errors yet
-			const sideChannel = this.additionalBlockChannelIds.get(sourceBranch.upperName + '|' + targetBranch.upperName)
-			if (sideChannel) {
-				this.slackMessages.postOrUpdate(changeInfo.source_cl, targetKey, {...message, channel: sideChannel})
+		if (targetBranch) {
+			const additionalChannelInfo = this.additionalBlockChannelIds.get(sourceBranch.upperName + '|' + targetBranch.upperName)
+			if (additionalChannelInfo) {
+				const [sideChannel, postOnlyToSideChannel] = additionalChannelInfo
+				if (!postOnlyToSideChannel) {
+					this.slackMessages.postOrUpdate(changeInfo.source_cl, targetBranch.name, message)
+				}
+				this.slackMessages.postOrUpdate(changeInfo.source_cl, targetBranch.name, { ...message, channel: sideChannel })
 			}
+			else {
+				this.slackMessages.postOrUpdate(changeInfo.source_cl, targetBranch.name, message)
+			}
+		}
+		else {
+			// separate message for syntax errors (other blockages have a target branch)
+			this.slackMessages.postOrUpdate(changeInfo.source_cl, blockage.failure.kind, message)
 		}
 
 		// Post message to owner in DM
@@ -845,10 +849,10 @@ export class BotNotifications implements BotEventHandler {
 	}
 
 	private readonly slackMessages?: SlackMessages
-	private readonly additionalBlockChannelIds = new Map<string, string>()
+	private readonly additionalBlockChannelIds = new Map<string, [string, boolean]>()
 }
 
-export function bindBotNotifications(events: BotEvents, slackChannelOverrides: [Branch, Branch, string][], persistence: Context, blockageUrlGenerator: NodeOpUrlGenerator, 
+export function bindBotNotifications(events: BotEvents, slackChannelOverrides: [Branch, Branch, string, boolean][], persistence: Context, blockageUrlGenerator: NodeOpUrlGenerator, 
 	externalUrl: string, logger: ContextualLogger) {
 	events.registerHandler(new BotNotifications(events.botname, events.botConfig.slackChannel, persistence, externalUrl, blockageUrlGenerator, logger, slackChannelOverrides))
 }
