[NO NON-PII COMMIT MESSAGE]

diff --git a/Engine/Extras/RoboMerge/v3/public/js/flow.js b/Engine/Extras/RoboMerge/v3/public/js/flow.js
--- a/Engine/Extras/RoboMerge/v3/public/js/flow.js
+++ b/Engine/Extras/RoboMerge/v3/public/js/flow.js
@@ -5,15 +5,17 @@ class Info {
         this.tooltip = tooltip;
         this.name = name;
         this.sources = sources;
-        this.forcedDests = new Set();
+        this.forcedEdges = new Set();
         this.defaultDests = new Set();
         this.blockAssetDests = new Set();
         this.importance = 0;
     }
     update(branch, getId) {
-        this.forcedDests = new Set([...this.forcedDests, ...branch.def.forceFlowTo.map(getId)]);
-        this.defaultDests = new Set([...this.defaultDests, ...branch.def.defaultFlow.map(getId)]);
-        this.blockAssetDests = new Set([...this.blockAssetDests, ...branch.def.blockAssetTargets.map(getId)]);
+        const edgeId = (alias) => branch.def.name + '->' + getId(alias);
+        // shuld really do the same for default and block assets
+        this.forcedEdges = new Set([...this.forcedEdges, ...branch.def.forceFlowTo.map(edgeId)]);
+        this.defaultDests = new Set([...this.defaultDests, ...branch.def.defaultFlow.map(edgeId)]);
+        this.blockAssetDests = new Set([...this.blockAssetDests, ...branch.def.blockAssetTargets.map(edgeId)]);
     }
 }
 const EDGE_STYLES = {
@@ -187,7 +189,8 @@ class Graph {
         const dstInfo = this.aliases.get(decoratedAlias(srcBranchStatus.bot, dst));
         ++srcInfo.importance;
         ++dstInfo.importance;
-        const isForced = srcInfo.forcedDests.has(dstInfo.id);
+        const hasEdge = (edges, target) => edges.has(srcBranchStatus.def.name + '->' + target.id);
+        const isForced = hasEdge(srcInfo.forcedEdges, dstInfo);
         if (!isForced && this.options.showOnlyForced) {
             return;
         }
@@ -195,8 +198,8 @@ class Graph {
         this.connectedNodes.add(dstInfo.id);
         const edgeStyle = srcBranchStatus.def.convertIntegratesToEdits ? 'roboshelf' :
             isForced ? 'forced' :
-                srcInfo.defaultDests.has(dstInfo.id) ? 'defaultFlow' :
-                    srcInfo.blockAssetDests.has(dstInfo.id) ? 'blockAssets' : 'onRequest';
+                hasEdge(srcInfo.defaultDests, dstInfo) ? 'defaultFlow' :
+                    hasEdge(srcInfo.blockAssetDests, dstInfo) ? 'blockAssets' : 'onRequest';
         const styles = [...EDGE_STYLES[edgeStyle]];
         const link = { src: srcInfo.id, dst: dstInfo.id, styles };
         this.links.push(link);
diff --git a/Engine/Extras/RoboMerge/v3/web/src/flow.ts b/Engine/Extras/RoboMerge/v3/web/src/flow.ts
--- a/Engine/Extras/RoboMerge/v3/web/src/flow.ts
+++ b/Engine/Extras/RoboMerge/v3/web/src/flow.ts
@@ -5,7 +5,7 @@ import { BranchDefForStatus, BranchStatus } from "../../src/robo/status-types"
 class Info {
 	graphNodeColor?: string
 
-	forcedDests = new Set<string>()
+	forcedEdges = new Set<string>()
 	defaultDests = new Set<string>()
 	blockAssetDests = new Set<string>()
 
@@ -16,9 +16,11 @@ class Info {
 	}
 
 	update(branch: BranchStatus, getId: (alias: string) => string) {
-		this.forcedDests = new Set([...this.forcedDests, ...branch.def.forceFlowTo.map(getId)])
-		this.defaultDests = new Set([...this.defaultDests, ...branch.def.defaultFlow.map(getId)])
-		this.blockAssetDests = new Set([...this.blockAssetDests, ...branch.def.blockAssetTargets.map(getId)])
+
+		const edgeId = (alias: string) => branch.def.name + '->' + getId(alias)
+		this.forcedEdges = new Set([...this.forcedEdges, ...branch.def.forceFlowTo.map(edgeId)])
+		this.defaultDests = new Set([...this.defaultDests, ...branch.def.defaultFlow.map(edgeId)])
+		this.blockAssetDests = new Set([...this.blockAssetDests, ...branch.def.blockAssetTargets.map(edgeId)])
 	}
 }
 
@@ -258,7 +260,10 @@ class Graph {
 		++srcInfo.importance
 		++dstInfo.importance
 
-		const isForced = srcInfo.forcedDests.has(dstInfo.id)
+		const hasEdge = (edges: Set<string>, target: Info) => 
+			edges.has(srcBranchStatus.def.name + '->' + target.id)
+
+		const isForced = hasEdge(srcInfo.forcedEdges, dstInfo)
 		if (!isForced && this.options.showOnlyForced) {
 			return
 		}
@@ -269,8 +274,8 @@ class Graph {
 		const edgeStyle =
 			srcBranchStatus.def.convertIntegratesToEdits ?	'roboshelf' :
 			isForced ?										'forced' :
-			srcInfo.defaultDests.has(dstInfo.id) ?			'defaultFlow' :
-			srcInfo.blockAssetDests.has(dstInfo.id) ?		'blockAssets' : 'onRequest'
+			hasEdge(srcInfo.defaultDests, dstInfo) ?		'defaultFlow' :
+			hasEdge(srcInfo.blockAssetDests, dstInfo) ?		'blockAssets' : 'onRequest'
 
 		const styles = [...EDGE_STYLES[edgeStyle]]
 
