Only don't change owner for robomerge user, not all bots

diff --git a/Engine/Extras/RoboMerge/v3/src/robo/edgebot.ts b/Engine/Extras/RoboMerge/v3/src/robo/edgebot.ts
--- a/Engine/Extras/RoboMerge/v3/src/robo/edgebot.ts
+++ b/Engine/Extras/RoboMerge/v3/src/robo/edgebot.ts
@@ -725,7 +725,7 @@ class EdgeBotImpl extends PerforceStatefulBot {
 		const log_str = (target.description || '').trim();
 		this.edgeBotLogger.info(`Submitted CL ${finalCl} to ${this.targetBranch.name}\n    ${log_str.replace(/\n/g, "\n    ")}`);
 
-		if (!isUserAKnownBot(info.author)) {
+		if (info.author != 'robomerge') {
 			// change owner, so users can edit change descriptions later for reconsideration
 			this.edgeBotLogger.info(`Setting owner of CL ${finalCl} to author of change: ${info.author}`);
 			try {
