Add ability to disregard block assets when appropriate

diff --git a/Engine/Extras/RoboMerge/v3/functional_tests/framework/src/tests/block-assets.ts b/Engine/Extras/RoboMerge/v3/functional_tests/framework/src/tests/block-assets.ts
--- a/Engine/Extras/RoboMerge/v3/functional_tests/framework/src/tests/block-assets.ts
+++ b/Engine/Extras/RoboMerge/v3/functional_tests/framework/src/tests/block-assets.ts
@@ -15,6 +15,7 @@ export class BlockAssets extends FunctionalTest {
 		await this.createStreamsAndWorkspaces(streams)
 
 		await P4Util.addFileAndSubmit(this.getClient('Main'), 'test.uasset', 'Initial content', true)
+		await P4Util.addFileAndSubmit(this.getClient('Main'), 'test2.uasset', 'Initial content', true)
 
 		const desc = 'Initial branch of files from Main'
 		await Promise.all([
@@ -23,29 +24,23 @@ export class BlockAssets extends FunctionalTest {
 		])
 	}
 
-	run() {
-		return P4Util.editFileAndSubmit(this.getClient('Main'), 'test.uasset', 'New content')
-	}	
+	async run() {
+		await P4Util.editFileAndSubmit(this.getClient('Main'), 'test.uasset', 'New content', '#disregardAssetBlock')
+		await P4Util.editFileAndSubmit(this.getClient('Main'), 'test2.uasset', 'New content')
+	}
 
 	verify() {
 		// situation before fixes:
 		//	- block asset on any target treated as syntax error, so blocks all
 		return Promise.all([
 			this.checkHeadRevision('Main', 'test.uasset', 2),
-			this.checkHeadRevision('Dev-Perkin', 'test.uasset', 1),
-
-
-// current state
-	
-/*
-	this.ensureBlocked('Main'),
-/*/
-// once fixed:
-
+			this.checkHeadRevision('Dev-Perkin', 'test.uasset', 2),
 			this.checkHeadRevision('Dev-Pootle', 'test.uasset', 2),
+			this.checkHeadRevision('Main', 'test2.uasset', 2),
+			this.checkHeadRevision('Dev-Perkin', 'test2.uasset', 1),
+			this.checkHeadRevision('Dev-Pootle', 'test2.uasset', 2),
 			this.ensureBlocked('Main', 'Dev-Perkin'),
 			this.ensureNotBlocked('Main', 'Dev-Pootle')
-/**/
 		])
 	}
 
diff --git a/Engine/Extras/RoboMerge/v3/src/new/graph.ts b/Engine/Extras/RoboMerge/v3/src/new/graph.ts
--- a/Engine/Extras/RoboMerge/v3/src/new/graph.ts
+++ b/Engine/Extras/RoboMerge/v3/src/new/graph.ts
@@ -641,7 +641,7 @@ Note first level is treated specially: for the purposes of these tests, force fl
 	return fail
 }
 
-type ChangeFlag = 'manual' | 'null' | 'ignore' | 'disregardexcludedauthors' | 'roboshelf'
+type ChangeFlag = 'manual' | 'null' | 'ignore' | 'disregardexcludedauthors' | 'roboshelf' | 'disregardassetblock'
 const ALLOWED_RAW_FLAGS = new Set(['null','ignore','deadend'])
 
 const FLAGMAP: {[name: string]: ChangeFlag} = {
diff --git a/Engine/Extras/RoboMerge/v3/src/robo/edgebot.ts b/Engine/Extras/RoboMerge/v3/src/robo/edgebot.ts
--- a/Engine/Extras/RoboMerge/v3/src/robo/edgebot.ts
+++ b/Engine/Extras/RoboMerge/v3/src/robo/edgebot.ts
@@ -408,6 +408,10 @@ class EdgeBotImpl extends PerforceStatefulBot {
 			description += '#ROBOMERGE[ALL]: #DISREGARDEXCLUDEDAUTHORS\n'
 		}
 
+		if (target.flags.has('disregardassetblock')) {
+			description += '#ROBOMERGE[ALL]: #DISREGARDASSETBLOCK\n'
+		}
+
 		let flags = '' // not all flags propagate, build them piecemeal
 
 		const thisBotMergeCommands = target.furtherMerges
diff --git a/Engine/Extras/RoboMerge/v3/src/robo/nodebot.ts b/Engine/Extras/RoboMerge/v3/src/robo/nodebot.ts
--- a/Engine/Extras/RoboMerge/v3/src/robo/nodebot.ts
+++ b/Engine/Extras/RoboMerge/v3/src/robo/nodebot.ts
@@ -1403,8 +1403,7 @@ export class NodeBot extends PerforceStatefulBot implements NodeBotInterface {
 			// check if we need to block a change containing assets
 			const blockAssetEdges: [EdgeBot, MergeAction][] = []
 			for (const action of info.targets) {
-				if (action.mergeMode === 'normal') {
-
+				if (action.mergeMode === 'normal' && !action.flags.has('disregardassetblock')) {
 
 					for (const blockTargetName of this.branch.blockAssetTargets) {
 						if (this._getBranch(blockTargetName) === action.branch) {
diff --git a/Engine/Extras/RoboMerge/v3/src/robo/targets.ts b/Engine/Extras/RoboMerge/v3/src/robo/targets.ts
--- a/Engine/Extras/RoboMerge/v3/src/robo/targets.ts
+++ b/Engine/Extras/RoboMerge/v3/src/robo/targets.ts
@@ -75,7 +75,7 @@ const ROBO_TAGS = [
 	'ROBOMERGE-SOURCE',
 ]
 
-type ChangeFlag = 'manual' | 'null' | 'ignore' | 'disregardexcludedauthors'
+type ChangeFlag = 'manual' | 'null' | 'ignore' | 'disregardexcludedauthors' | 'disregardassetblock'
 
 // mapping of #ROBOMERGE: flags to canonical names
 // use these with a pound like #ROBOMERGE: #stage
@@ -94,6 +94,9 @@ const FLAGMAP: {[name: string]: ChangeFlag} = {
 
 	// process this change even if the author is on the excluded list
 	disregardexcludedauthors: 'disregardexcludedauthors',
+
+	// process this change even if block assets is enabled for the edge and it contains assets
+	disregardassetblock: 'disregardassetblock',
 }
 
 const ALLOWED_RAW_FLAGS = ['null','ignore','deadend']
