[NO NON-PII COMMIT MESSAGE]

diff --git a/Engine/Extras/RoboMerge/v3/src/robo/nodebot.ts b/Engine/Extras/RoboMerge/v3/src/robo/nodebot.ts
--- a/Engine/Extras/RoboMerge/v3/src/robo/nodebot.ts
+++ b/Engine/Extras/RoboMerge/v3/src/robo/nodebot.ts
@@ -1442,7 +1442,9 @@ export class NodeBot extends PerforceStatefulBot implements NodeBotInterface {
 						const failure: Failure = {kind: 'Disallowed files', description: 'Changelist contains assets which are not allowed on this route'}
 						const pending: PendingChange = {change: info, action, newCl: -1}
 
-						await edge.block(this.createEdgeBlockageInfo(action.branch, failure, pending))
+						if (!change.isUserRequest) {
+							await edge.block(this.createEdgeBlockageInfo(action.branch, failure, pending))
+						}
 
 						// ... then let node do facilitate notification handling
 						await this.handleMergeFailure(failure, pending)
