updating roborepl

- Fixing reconsider when using just a Node
- adding some new methods to stop() and start() the server, restart a bot restart_bot(bot), is_running(bot)
- allow to prompt for the password instead of passing as plain text in the terminal

diff --git a/Engine/Extras/RoboMerge/v3/repl/roborepl.py b/Engine/Extras/RoboMerge/v3/repl/roborepl.py
--- a/Engine/Extras/RoboMerge/v3/repl/roborepl.py
+++ b/Engine/Extras/RoboMerge/v3/repl/roborepl.py
@@ -1,6 +1,6 @@
 import code
-
 import json, requests, sys
+from getpass import getpass
 
 if sys.platform == 'darwin':
 	import gnureadline
@@ -150,7 +150,9 @@ def fetch_branch_data():
 
 
 @export('login')
-def login(user, password):
+def login(user, password=None):
+	if not password:
+		password = getpass('Enter your password: ')
 
 	r = requests.post(f'{ROBOMERGE_URL}/dologin', data = dict(
 		user = user,
@@ -222,6 +224,13 @@ def list_blockages(bot):
 		if edge.blockage:
 			print(f"{edge} {error('blocked')} at CL {edge.blockage.cl} (owner {edge.blockage.owner})")
 
+@export('status')
+def is_running(bot):
+	result = get(f"{ROBOMERGE_URL}/api/control/isrunning/{bot}")
+	if result.status_code != 200:
+		raise Exception(result.text)	
+	print(f"{result.text}")
+
 # operations
 
 def op_url(node_or_edge, op):
@@ -229,7 +238,7 @@ def op_url(node_or_edge, op):
 	suffix = '/op/' + op
 
 	if isinstance(node_or_edge, Node):
-		return prefix + f'node/{node_or_edge.name()}' + suffix
+		return prefix + f'node/{node_or_edge.name}' + suffix
 
 	edge = node_or_edge
 	return prefix + f'node/{edge.source.name.upper()}/edge/{edge.target.name.upper()}' + suffix
@@ -255,5 +264,23 @@ def reconsider(node_or_edge, cl):
 	if result.status_code != 200:
 		raise Exception(result.text)
 
+@export('control')
+def stop():
+	result = post(f"{ROBOMERGE_URL}/api/control/stop")
+	if result.status_code != 200:
+		raise Exception(result.text)
+
+@export('control')
+def start():
+	result = post(f"{ROBOMERGE_URL}/api/control/start")
+	if result.status_code != 200:
+		raise Exception(result.text)
+
+@export('control')
+def restart_bot(bot):
+	result = post(f"{ROBOMERGE_URL}/api/control/restart-bot/{bot}")
+	if result.status_code != 200:
+		raise Exception(result.text)
+
 if __name__ == '__main__':
 	code.interact('RoboMerge REPL v0.1' + ENDC, None, context)
