Add a branchspec functional test

diff --git a/Engine/Extras/RoboMerge/v3/functional_tests/framework/src/framework.ts b/Engine/Extras/RoboMerge/v3/functional_tests/framework/src/framework.ts
--- a/Engine/Extras/RoboMerge/v3/functional_tests/framework/src/framework.ts
+++ b/Engine/Extras/RoboMerge/v3/functional_tests/framework/src/framework.ts
@@ -306,6 +306,16 @@ export abstract class FunctionalTest {
 		)
 	}
 
+	branchSpec(name: string, view:string) {
+		return P4Util.specFormat(
+			["Branch", name],
+			["Owner", "root"],
+			["Description", "\n\tCreated by root."],
+			["Options", "unlocked"],
+			["View", view]
+		)
+	}
+
 	protected makeBranchDef(stream: string, to: string[], forceAll?: boolean): RobomergeBranchSpec {
 		const name = this.fullBranchName(stream)
 		return {
diff --git a/Engine/Extras/RoboMerge/v3/functional_tests/framework/src/test-perforce.ts b/Engine/Extras/RoboMerge/v3/functional_tests/framework/src/test-perforce.ts
--- a/Engine/Extras/RoboMerge/v3/functional_tests/framework/src/test-perforce.ts
+++ b/Engine/Extras/RoboMerge/v3/functional_tests/framework/src/test-perforce.ts
@@ -27,10 +27,11 @@ export class Perforce {
 		return this.exec(process.platform === 'darwin' ? ['-plocalhost:1666', 'login', '-s'] : ['login', '-s'])
 	}
 
-	depot(depotType: string, spec: string)	{ return this.exec(['depot', '-t', depotType],	{stdin:spec}) }
-	user(spec: string)						{ return this.exec(['user', '-fi'],				{stdin:spec}) }
-	stream(spec: string)					{ return this.exec(['stream'],					{stdin:spec}) }
-	client(user: string, spec: string)		{ return this.exec(['-u', user, 'client'],		{stdin:spec}) }
+	depot(depotType: string, spec: string)  { return this.exec(['depot', '-t', depotType],  {stdin:spec}) }
+	user(spec: string)                      { return this.exec(['user', '-fi'],             {stdin:spec}) }
+	stream(spec: string)                    { return this.exec(['stream'],                  {stdin:spec}) }
+	client(user: string, spec: string)      { return this.exec(['-u', user, 'client'],      {stdin:spec}) }
+	branch(spec: string)                    { return this.exec(['branch', '-i'],            {stdin:spec}) }
 
 	add(user: string, workspace: string, path: string, binary?: boolean)
 	{
diff --git a/Engine/Extras/RoboMerge/v3/functional_tests/framework/src/tests.ts b/Engine/Extras/RoboMerge/v3/functional_tests/framework/src/tests.ts
--- a/Engine/Extras/RoboMerge/v3/functional_tests/framework/src/tests.ts
+++ b/Engine/Extras/RoboMerge/v3/functional_tests/framework/src/tests.ts
@@ -6,6 +6,7 @@ import { Perforce } from './test-perforce'
 import { GenericTargetTest } from './GenericTargetTest'
 import { BlockAssets } from './tests/block-assets'
 import { BlockIgnore } from './tests/block-ignore'
+import { BranchspecTest } from './tests/branchspec'
 import { ConfirmBinaryStomp } from './tests/confirm-binary-stomp'
 import { ConfirmTextResolve } from './tests/confirm-text-resolve'
 import { ConfirmTextResolveBinaryStomp } from './tests/confirm-text-resolve-binary-stomp'
@@ -216,6 +217,7 @@ async function go() {
 		new UnreachableSkip(p4), // 45
 
 		new PostToAdditionalChannel(p4),
+		new BranchspecTest(p4)
 	]
 
 	const TARGET_TEST_DEFS: [string[], string | null][] = [
diff --git a/Engine/Extras/RoboMerge/v3/functional_tests/framework/src/tests/branchspec.ts b/Engine/Extras/RoboMerge/v3/functional_tests/framework/src/tests/branchspec.ts
new file mode 100644
--- /dev/null
+++ b/Engine/Extras/RoboMerge/v3/functional_tests/framework/src/tests/branchspec.ts
@@ -0,0 +1,48 @@
+// Copyright Epic Games, Inc. All Rights Reserved.
+import { P4Util } from '../framework'
+import { MultipleDevAndReleaseTestBase } from '../MultipleDevAndReleaseTestBase'
+
+
+export class BranchspecTest extends MultipleDevAndReleaseTestBase {
+	async setup() {
+		await this.p4.depot('stream', this.depotSpec())
+		await this.createStreamsAndWorkspaces()
+		await this.p4.branch(this.branchSpec('Main->Pootle', `${this.getStreamPath('Main')}/Allowed_* ${this.getStreamPath('Dev-Pootle')}/Allowed_*`))
+
+		await P4Util.addFileAndSubmit(this.getClient('Main'), 'Allowed_test.uasset', 'Initial content', true)
+		await P4Util.addFileAndSubmit(this.getClient('Main'), 'Filtered_test.uasset', 'Initial content', true)
+
+		await this.initialPopulate()
+	}
+
+	async run() {
+		const r1client = this.getClient('Release-1.0')
+		await r1client.sync()
+		await P4Util.editFile(r1client, 'Allowed_test.uasset', 'Initial content\nSome more!')
+		await P4Util.editFile(r1client, 'Filtered_test.uasset', 'Initial content\nSome more!')
+		await P4Util.submit(r1client, "Files")
+	}
+
+	async verify() {
+
+		this.info('Ensuring branchspec filtered as expected')
+
+		return Promise.all([
+			this.checkHeadRevision('Main', 'Allowed_test.uasset', 2),
+			this.checkHeadRevision('Dev-Perkin', 'Allowed_test.uasset', 2),
+			this.checkHeadRevision('Dev-Pootle', 'Allowed_test.uasset', 2),
+			this.checkHeadRevision('Main', 'Filtered_test.uasset', 2),
+			this.checkHeadRevision('Dev-Perkin', 'Filtered_test.uasset', 2),
+			this.checkHeadRevision('Dev-Pootle', 'Filtered_test.uasset', 1),
+		])
+	}
+
+	getEdges() {
+		return [{
+			from: this.fullBranchName('Main'),
+			to: this.fullBranchName('Dev-Pootle'),
+			branchspec: 'Main->Pootle'
+		}]
+	}
+
+}
