[NO NON-PII COMMIT MESSAGE]

diff --git a/Engine/Extras/RoboMerge/v3/src/robo/session.ts b/Engine/Extras/RoboMerge/v3/src/robo/session.ts
--- a/Engine/Extras/RoboMerge/v3/src/robo/session.ts
+++ b/Engine/Extras/RoboMerge/v3/src/robo/session.ts
@@ -241,17 +241,17 @@ export class Session {
 	}
 
 	static _testTokenSigning(logger: ContextualLogger) {
-		const TEST_DATA: AuthData = {user: 'marlin.kingsly', displayName: '-', tags: new Set(['admin'])}
+		const TEST_DATA: AuthData = {user: '<p4 user name>', displayName: '-', tags: new Set(['admin'])}
 		
 		Session._testSingleToken('match', TEST_DATA, logger)
-		Session._testSingleToken('mismatched user', TEST_DATA, logger, {user: 'marlin.kngsly', displayName: '-', tags: new Set(['admin'])})
-		Session._testSingleToken('mismatched tag', TEST_DATA, logger, {user: 'marlin.kingsly', displayName: '-', tags: new Set(['admn'])})
+		Session._testSingleToken('mismatched user', TEST_DATA, logger, {user: '<p4 user name>', displayName: '-', tags: new Set(['admin'])})
+		Session._testSingleToken('mismatched tag', TEST_DATA, logger, {user: '<p4 user name>', displayName: '-', tags: new Set(['admn'])})
 		// actually always going to encode user as lower case
-		Session._testSingleToken('mismatched user case', TEST_DATA, logger, {user: 'Marlin.Kingsly', displayName: '-', tags: new Set(['admin'])})
-		Session._testSingleToken('mismatched tag, case only', TEST_DATA, logger, {user: 'marlin.kingsly', displayName: '-', tags: new Set(['ADMIN'])})
-		Session._testSingleToken('match, no tags', {user: 'marlin.kingsly', displayName: '-', tags: new Set([])}, logger)
+		Session._testSingleToken('mismatched user case', TEST_DATA, logger, {user: '<p4 user name>', displayName: '-', tags: new Set(['admin'])})
+		Session._testSingleToken('mismatched tag, case only', TEST_DATA, logger, {user: '<p4 user name>', displayName: '-', tags: new Set(['ADMIN'])})
+		Session._testSingleToken('match, no tags', {user: '<p4 user name>', displayName: '-', tags: new Set([])}, logger)
 		// (~~~~ forces a + in base 64 which gets replaced with a -). Quite a few of the MACs generated also test replacement
-		Session._testSingleToken('match, with replaced character', {user: 'marlin.kingsly', displayName: '-', tags: new Set(['~~~~'])}, logger)
+		Session._testSingleToken('match, with replaced character', {user: '<p4 user name>', displayName: '-', tags: new Set(['~~~~'])}, logger)
 	}
 }
 
