add some debug loging for slack issues

diff --git a/Engine/Extras/RoboMerge/v3/src/robo/slack.ts b/Engine/Extras/RoboMerge/v3/src/robo/slack.ts
--- a/Engine/Extras/RoboMerge/v3/src/robo/slack.ts
+++ b/Engine/Extras/RoboMerge/v3/src/robo/slack.ts
@@ -148,7 +148,7 @@ export class Slack {
 
 		const result = JSON.parse(rawResult)
 		if (!result.ok) {
-			throw new Error(rawResult)
+			throw new Error(`url: '${url}' error: '${rawResult}'`)
 		}
 		return result
 	}
