just pass the email when calling users.lookupByEmail

diff --git a/Engine/Extras/RoboMerge/v3/src/robo/slack.ts b/Engine/Extras/RoboMerge/v3/src/robo/slack.ts
--- a/Engine/Extras/RoboMerge/v3/src/robo/slack.ts
+++ b/Engine/Extras/RoboMerge/v3/src/robo/slack.ts
@@ -99,7 +99,7 @@ export class Slack {
 	}
 
 	async lookupUserIdByEmail(email: string) : Promise<string> {
-		return (await this.get('users.lookupByEmail', {token: this.channel.botToken, email})).user.id
+		return (await this.get('users.lookupByEmail', {email})).user.id
 	}
 
 	async openDMConversation(users: string | string[]) : Promise<string> {
