Fix up issues required to compile with latest tsc/node
Fix client being created with wrong root directory on win32

diff --git a/Engine/Extras/RoboMerge/v3/.vscode/launch.json b/Engine/Extras/RoboMerge/v3/.vscode/launch.json
--- a/Engine/Extras/RoboMerge/v3/.vscode/launch.json
+++ b/Engine/Extras/RoboMerge/v3/.vscode/launch.json
@@ -17,13 +17,14 @@
             "console": "integratedTerminal",
             "outputCapture": "std",
             "args": [
-                "-noIPC"
+                "-noIPC",
+                "-noTLS"
             ],
             "env": {
                 "ROBO_DEV_MODE": "true",
                 "ROBO_EXTERNAL_URL": "http://localhost:8877",
                 "NODE_ENV": "development",
-                "ROBO_BRANCHSPECS_DIRECTORY": "./debugdata",
+                "ROBO_BRANCHSPECS_DIRECTORY": "../Robomerge/data",
                 "ROBO_VAULT_PATH": "./vault",
                 "EPIC_ENV": "dev",
                 "EPIC_DEPLOYMENT": "local",
diff --git a/Engine/Extras/RoboMerge/v3/src/common/perforce.ts b/Engine/Extras/RoboMerge/v3/src/common/perforce.ts
--- a/Engine/Extras/RoboMerge/v3/src/common/perforce.ts
+++ b/Engine/Extras/RoboMerge/v3/src/common/perforce.ts
@@ -705,7 +705,7 @@ export class PerforceContext {
 
 	// Create a new workspace for Robomerge GraphBot
 	async newGraphBotWorkspace(name: string, extraParams: any, edgeServer?: {id: string, address: string}) {
-		return this.newWorkspace(name, {Root: '/src/' + name, ...extraParams}, edgeServer);
+		return this.newWorkspace(name, {Root: getRootDirectoryForBranch(name), ...extraParams}, edgeServer);
 	}
 
 	// Create a new workspace for Robomerge to read branchspecs from
@@ -1365,7 +1365,7 @@ export class PerforceContext {
 			options.env = {};
 			for (let key in process.env)
 				options.env[key] = process.env[key];
-			options.env.PWD = path.resolve(options.cwd);
+			options.env.PWD = path.resolve(options.cwd.toString());
 		}
 
 		const doExecFile = function(retries: number): Promise<string> {
diff --git a/Engine/Extras/RoboMerge/v3/src/common/webserver.ts b/Engine/Extras/RoboMerge/v3/src/common/webserver.ts
--- a/Engine/Extras/RoboMerge/v3/src/common/webserver.ts
+++ b/Engine/Extras/RoboMerge/v3/src/common/webserver.ts
@@ -198,7 +198,7 @@ export class WebServer {
 			throw new Error(`Unknown file type '${filetype}'`)
 		}
 
-		let encoding: string | null = null
+		let encoding: BufferEncoding | null = null
 		let mimeType = MIME_TYPES.get(filetype)
 		if (mimeType) {
 			// assuming known file types other than images are text
@@ -248,7 +248,7 @@ export class WebServer {
 		return headers;	
 	}
 
-	private handleFileRequest(response: http.ServerResponse, filePattern: string, encoding: string | null,
+	private handleFileRequest(response: http.ServerResponse, filePattern: string, encoding: BufferEncoding | null,
 								mimeType: string, route: RegExp, path: string, inHeaders?: [string, string][]) {
 		const filePath = path.replace(route, filePattern)
 		if (filePath.search(/\.\./) !== -1) {
diff --git a/Engine/Extras/RoboMerge/v3/src/robo/bot-interfaces.ts b/Engine/Extras/RoboMerge/v3/src/robo/bot-interfaces.ts
--- a/Engine/Extras/RoboMerge/v3/src/robo/bot-interfaces.ts
+++ b/Engine/Extras/RoboMerge/v3/src/robo/bot-interfaces.ts
@@ -52,7 +52,7 @@ export interface IPCControls {
 }
 
 export interface EdgeBotInterface {
-	get ipcControls(): IPCControls
+	getIPCControls(): IPCControls
 
 	bypassGateWindow(sense: boolean): void
 }
diff --git a/Engine/Extras/RoboMerge/v3/src/robo/edgebot.ts b/Engine/Extras/RoboMerge/v3/src/robo/edgebot.ts
--- a/Engine/Extras/RoboMerge/v3/src/robo/edgebot.ts
+++ b/Engine/Extras/RoboMerge/v3/src/robo/edgebot.ts
@@ -1117,7 +1117,7 @@ export class EdgeBot
 		return this.impl.implicitCommands
 	}
 
-	get ipcControls(): IPCControls {
+	getIPCControls(): IPCControls {
 		return {
 			block: this.block,
 			unblock: this.unblock,
diff --git a/Engine/Extras/RoboMerge/v3/src/robo/ipc.ts b/Engine/Extras/RoboMerge/v3/src/robo/ipc.ts
--- a/Engine/Extras/RoboMerge/v3/src/robo/ipc.ts
+++ b/Engine/Extras/RoboMerge/v3/src/robo/ipc.ts
@@ -332,7 +332,7 @@ export class IPC {
 		let edge: EdgeBotInterface | null = null
 		if (edgeName) {
 			edge = bot.getImmediateEdge(edgeName)
-			const edgeIPC = edge && edge.ipcControls
+			const edgeIPC = edge && edge.getIPCControls()
 			if (!edgeIPC) {
 				return {statusCode: 400, message: `Node ${nodeName} does not have edge ${edgeName}`}
 			}
diff --git a/Engine/Extras/RoboMerge/v3/src/robo/nodebot.ts b/Engine/Extras/RoboMerge/v3/src/robo/nodebot.ts
--- a/Engine/Extras/RoboMerge/v3/src/robo/nodebot.ts
+++ b/Engine/Extras/RoboMerge/v3/src/robo/nodebot.ts
@@ -187,7 +187,7 @@ export class NodeBot extends PerforceStatefulBot implements NodeBotInterface {
 	getEdgeIPCControls(branch: BranchArg): IPCControls | null {
 		const edge = this.getImmediateEdge(branch)
 
-		return edge ? edge.ipcControls : null  
+		return edge ? edge.getIPCControls() : null  
 	}
 
 	hasEdge(branch: BranchArg): boolean {
diff --git a/Engine/Extras/RoboMerge/v3/src/robo/watchdog.ts b/Engine/Extras/RoboMerge/v3/src/robo/watchdog.ts
--- a/Engine/Extras/RoboMerge/v3/src/robo/watchdog.ts
+++ b/Engine/Extras/RoboMerge/v3/src/robo/watchdog.ts
@@ -275,7 +275,7 @@ class Watchdog {
 			this.childExited()
 			this.respawn()
 		})
-		this.child.on('message', (msg) => {
+		this.child.on('message', (msg: any) => {
 			if (this.child === null)
 				return
 			
