Edge based resolvers

diff --git a/Engine/Extras/RoboMerge/v3/src/robo/branchdefs.ts b/Engine/Extras/RoboMerge/v3/src/robo/branchdefs.ts
--- a/Engine/Extras/RoboMerge/v3/src/robo/branchdefs.ts
+++ b/Engine/Extras/RoboMerge/v3/src/robo/branchdefs.ts
@@ -136,6 +136,8 @@ const edgeOptionFieldsPrototype = {
 
 	postOnlyToAdditionalChannel: false,
 
+	resolver: '',
+
 	terminal: false, // changes go along terminal edges but no further
 
 	implicitCommands: [''],
@@ -431,12 +433,7 @@ export class BranchDefs {
 
 				// check all properties are known (could make things case insensitive here)
 				for (const keyName of Object.keys(edge)) {
-					if (keyName !== 'from' && keyName !== 'to' 
-
-// temporarily allow resolver until supported properly
-&& keyName !== 'resolver'
-
-						&& !edgeOptionFieldNames.has(keyName)) {
+					if (keyName !== 'from' && keyName !== 'to' && !edgeOptionFieldNames.has(keyName)) {
 						throw new Error(`Unknown property '${keyName}' specified for edge ${edge.from}->${edge.to}`)
 					}
 				}
diff --git a/Engine/Extras/RoboMerge/v3/src/robo/edgebot.ts b/Engine/Extras/RoboMerge/v3/src/robo/edgebot.ts
--- a/Engine/Extras/RoboMerge/v3/src/robo/edgebot.ts
+++ b/Engine/Extras/RoboMerge/v3/src/robo/edgebot.ts
@@ -175,6 +175,10 @@ class EdgeBotImpl extends PerforceStatefulBot {
 		return this.options.implicitCommands || []
 	}
 
+	get resolver() {
+		return this.options.resolver
+	}
+
 	get isTerminal() {
 		return !!this.options.terminal
 	}
@@ -949,8 +953,8 @@ class EdgeBotImpl extends PerforceStatefulBot {
 			this.pauseState.applyStatus(status)
 		}
 
-		if (this.targetBranch.resolver) {
-			status.resolver = this.targetBranch.resolver
+		if (this.options.resolver) {
+			status.resolver = this.options.resolver
 		}
 		status.disallowSkip = this.options.disallowSkip
 		status.incognitoMode = this.options.incognitoMode
@@ -1096,6 +1100,7 @@ export class EdgeBot
 	get disallowSkip() { return this.impl.disallowSkip }
 	get incognitoMode() { return this.impl.incognitoMode }
 	get excludedAuthors() { return this.impl.excludedAuthors }
+	get resolver() { return this.impl.resolver }
 	get isActive() { return this.impl.isActive }
 	get isAvailable() { return this.impl.isAvailable }
 	get isBlocked() { return this.impl.isBlocked }
diff --git a/Engine/Extras/RoboMerge/v3/src/robo/targets.ts b/Engine/Extras/RoboMerge/v3/src/robo/targets.ts
--- a/Engine/Extras/RoboMerge/v3/src/robo/targets.ts
+++ b/Engine/Extras/RoboMerge/v3/src/robo/targets.ts
@@ -44,7 +44,15 @@ export function getIntegrationOwner(arg0: Branch | PendingChange, overriddenOwne
 	const branch = pending ? pending.action.branch : targetBranch
 	const owner = pending ? pending.change.owner : overriddenOwner
 
-	return branch!.resolver || owner || null
+	let resolver = branch!.resolver
+	if (pending) {
+		const edgeprops = pending.change.branch.edgeProperties.get(branch!.upperName)
+		if (edgeprops) {
+			resolver = edgeprops.resolver || null
+		}	
+	}
+
+	return resolver || owner || null
 }
 
 export function getNodeBotFullName(botname: string, branchName: string) {
