Remove end of line indicator from exclusive checkout regex as on windows there is a \r that causes it to fail to match

diff --git a/Engine/Extras/RoboMerge/v3/src/common/perforce.ts b/Engine/Extras/RoboMerge/v3/src/common/perforce.ts
--- a/Engine/Extras/RoboMerge/v3/src/common/perforce.ts
+++ b/Engine/Extras/RoboMerge/v3/src/common/perforce.ts
@@ -16,7 +16,7 @@ const RETRY_ERROR_MESSAGES = [
 ]
 
 const INTEGRATION_FAILURE_REGEXES: [RegExp, string][] = [
-	[/^(.*[\\\/])(.*) - can't \w+ exclusive file already opened$/, 'partial_integrate'],
+	[/^(.*[\\\/])(.*) - can't \w+ exclusive file already opened/, 'partial_integrate'],
 	[/Move\/delete\(s\) must be integrated along with matching move\/add\(s\)/, 'split_move'],
 ]
 
