Log who paused

diff --git a/Engine/Extras/RoboMerge/v3/src/robo/perforce-stateful-bot.ts b/Engine/Extras/RoboMerge/v3/src/robo/perforce-stateful-bot.ts
--- a/Engine/Extras/RoboMerge/v3/src/robo/perforce-stateful-bot.ts
+++ b/Engine/Extras/RoboMerge/v3/src/robo/perforce-stateful-bot.ts
@@ -75,7 +75,7 @@ export abstract class PerforceStatefulBot implements BotIPC {
 		}
 
 		this.pauseState.manuallyPause(message, owner)
-		this.logger.info(`Paused ${this.fullName}: ${message}.`)
+		this.logger.info(`Paused ${this.fullName} requested by ${owner}: ${message}.`)
 	}
 	unpause(requester: string) {
 		if (!this.isManuallyPaused) {
