Minor language update

diff --git a/Engine/Extras/RoboMerge/v3/public/createshelf.html b/Engine/Extras/RoboMerge/v3/public/createshelf.html
--- a/Engine/Extras/RoboMerge/v3/public/createshelf.html
+++ b/Engine/Extras/RoboMerge/v3/public/createshelf.html
@@ -29,8 +29,8 @@
     <div id="result"></div>
     <div id="nextSteps" class="initiallyHidden centered-text">
         <h3>What next?</h3>
-        <p>Robomerge will run the merge for the changelist again and create any shelved conflicts in the workspace you chose. You should receive an email when this process is complete.</p>
-        <p>If you do not receive an email, please <strong><a href="/contact">Contact Us</a></strong> for Robomerge help.</p>
+        <p>Robomerge will run the merge for the changelist again and create any shelved conflicts in the workspace you chose. You should receive a notification when this process is complete.</p>
+        <p>If you do not receive a notification, please <strong><a href="/contact">Contact Us</a></strong> for Robomerge help.</p>
     </div>
     <div id="createStreamWorkspaceTutorial" class="initiallyHidden centered-text">
         <h2>Sorry, no user workspaces found for stream "<strong class="streamName"></strong>". Robomerge cannot create a shelf for you without one.</h2>
