Allow shelf creation for commit failures

diff --git a/Engine/Extras/RoboMerge/v3/public/js/boilerplate.js b/Engine/Extras/RoboMerge/v3/public/js/boilerplate.js
--- a/Engine/Extras/RoboMerge/v3/public/js/boilerplate.js
+++ b/Engine/Extras/RoboMerge/v3/public/js/boilerplate.js
@@ -1387,13 +1387,18 @@ function renderActionsCell_Common(actionCell, data, operationFunction, operation
 				window.location.href = stompRequest;
 			}, `Use ${conflict.cl} to stomp binary changes in ${conflict.target}`);
 
-			// Can only create shelves or perform stomps for merge conflicts
-			if (conflict.kind !== 'Merge conflict') {
+			// Can only create shelves for merge conflicts and commit failures
+			if (conflict.kind !== 'Merge conflict' &&
+				conflict.kind !== 'Commit failure') {
 				shelfOption.addClass("disabled")
 				shelfOption.off('click')
 				// Bootstrap magic converts 'title' to 'data-original-title'
 				shelfOption.attr('data-original-title', `Shelving not available for ${conflict.kind.toLowerCase()}.`)
+			}
 
+			// Can only perform stomps for merge conflicts
+			if (conflict.kind !== 'Merge conflict')
+			{
 				stompOption.addClass("disabled")
 				stompOption.off('click')
 				stompOption.attr('data-original-title', `Stomp not available for ${conflict.kind.toLowerCase()}.`)
diff --git a/Engine/Extras/RoboMerge/v3/web/src/branches.ts b/Engine/Extras/RoboMerge/v3/web/src/branches.ts
--- a/Engine/Extras/RoboMerge/v3/web/src/branches.ts
+++ b/Engine/Extras/RoboMerge/v3/web/src/branches.ts
@@ -792,13 +792,17 @@ function renderActionsCell_Common(actionCell: JQuery, data: Partial<BotStatusFie
 				window.location.href = stompRequest;
 			}, `Use ${conflict.cl} to stomp binary changes in ${conflict.target}`);
 
-			// Can only create shelves or perform stomps for merge conflicts
-			if (conflict.kind !== 'Merge conflict') {
+			// Can only create shelves for merge conflicts and commit failures
+			if (conflict.kind !== 'Merge conflict' &&
+			    conflict.kind !== 'Commit failure') {
 				shelfOption.addClass("disabled")
 				shelfOption.off('click')
 				// Bootstrap magic converts 'title' to 'data-original-title'
 				shelfOption.attr('data-original-title', `Shelving not available for ${conflict.kind.toLowerCase()}.`)
+			}
 
+			// Can only perform stomps for merge conflicts
+			if (conflict.kind !== 'Merge conflict') {
 				stompOption.addClass("disabled")
 				stompOption.off('click')
 				stompOption.attr('data-original-title', `Stomp not available for ${conflict.kind.toLowerCase()}.`)
