Fixup slack channels referenced

diff --git a/Engine/Extras/RoboMerge/v3/email_templates/generic-email-template.html b/Engine/Extras/RoboMerge/v3/email_templates/generic-email-template.html
--- a/Engine/Extras/RoboMerge/v3/email_templates/generic-email-template.html
+++ b/Engine/Extras/RoboMerge/v3/email_templates/generic-email-template.html
@@ -6,5 +6,5 @@ ${message}
 </pre>
 <h3><a href="${robomergeRootUrl}#${botname}">Proceed to Robomerge webpage</a></h3>
 <footer>
-	<b>Need Help?</b> Send a message to <a href="https://epic.slack.com/messages/robomerge-help/"><code>#robomerge-help</code></a> on Slack.
+	<b>Need Help?</b> Send a message to <a href="https://epic.slack.com/messages/dt-robomerge-help-ext/"><code>#dt-robomerge-help-ext</code></a> on Slack.
 </footer>
diff --git a/Engine/Extras/RoboMerge/v3/email_templates/new-blockage.html b/Engine/Extras/RoboMerge/v3/email_templates/new-blockage.html
--- a/Engine/Extras/RoboMerge/v3/email_templates/new-blockage.html
+++ b/Engine/Extras/RoboMerge/v3/email_templates/new-blockage.html
@@ -9,5 +9,5 @@ ${resolutionOptions}
 <br />
 <footer>
 	<a href="${robomergeRootUrl}#${botname}">${botname} Homepage on Robomerge</a><br />
-	<b>Need Help?</b> Send a message to <a href="https://epic.slack.com/messages/robomerge-help/"><code>#robomerge-help</code></a> on Slack.
+	<b>Need Help?</b> Send a message to <a href="https://epic.slack.com/messages/dt-robomerge-help-ext/"><code>#dt-robomerge-help-ext</code></a> on Slack.
 </footer>
diff --git a/Engine/Extras/RoboMerge/v3/public/stomp.html b/Engine/Extras/RoboMerge/v3/public/stomp.html
--- a/Engine/Extras/RoboMerge/v3/public/stomp.html
+++ b/Engine/Extras/RoboMerge/v3/public/stomp.html
@@ -31,7 +31,7 @@
             <div class="card-body centered-text">
                 <p>When a conflict consists of only binary files, Robomerge will offer to 'stomp' those binary files in the target path (<strong class="targetName"></strong>) with the revisions in the source path (<strong class="sourceName"></strong>). If your conflict consists of more than just binary files, Robomerge will not be able to provide this functionality.</p>
                 <p>Once your stomp request is verified and you elect to use stomp, you should receive an email when the process is complete.</p>
-                <p>Please report any issues to <a href="https://epic.slack.com/messages/robomerge-help/"><code>#robomerge-help</code></a> on Slack.</p>
+                <p>Please report any issues to <a href="https://epic.slack.com/messages/dt-robomerge-help-ext/"><code>#dt-robomerge-help-ext</code></a> on Slack.</p>
             </div>
         </div>
     </div>
