[NO NON-PII COMMIT MESSAGE]

diff --git a/Engine/Extras/RoboMerge/v3/src/common/perforce.ts b/Engine/Extras/RoboMerge/v3/src/common/perforce.ts
--- a/Engine/Extras/RoboMerge/v3/src/common/perforce.ts
+++ b/Engine/Extras/RoboMerge/v3/src/common/perforce.ts
@@ -485,11 +485,17 @@ export class PerforceContext {
 
 	/** get a single change in the format of changes() */
 	async getChange(path_in: string, changenum: number, status?: ChangelistStatus) {
-		const list = await this.changes(`${path_in}@${changenum},${changenum}`, -1, 1, status)
+		const list = await this.changes(`${path_in}@${changenum},${changenum}`, -1, 1, status) as Change[]
 		if (list.length <= 0) {
 			throw new Error(`Could not find changelist ${changenum} in ${path_in}`);
 		}
-		return <Change>list[0];
+		if (list.length > 1 || list[0].change !== changenum) {
+			// log for now
+			this.logger.error('p4.getChange unexpected result' +
+				list.map(change => `\n    ${change.change}: user ${change.user}, workspace ${change.client}`).join('')
+			)
+		}
+		return list[0]
 	}
 
 	/**
