Reorganize slack messages API to allow issuing slack messages from bots and begin shift away from email
Tag and invite owner when conflict is posted to channel
Nagging and triaging
Provide syntax and commit errors in channel
Shelf created notifications come via slack
Use slack to notify user of syntax error and commit error on reconsider

diff --git a/Engine/Extras/RoboMerge/v3/src/common/perforce.ts b/Engine/Extras/RoboMerge/v3/src/common/perforce.ts
--- a/Engine/Extras/RoboMerge/v3/src/common/perforce.ts
+++ b/Engine/Extras/RoboMerge/v3/src/common/perforce.ts
@@ -225,7 +225,7 @@ export interface Change {
 	isUserRequest?: boolean;
 	ignoreExcludedAuthors?: boolean
 	forceCreateAShelf?: boolean
-	sendNoShelfEmail?: boolean // Used for requesting stomps for internal RM usage, such as stomp changes
+	sendNoShelfNotification?: boolean // Used for requesting stomps for internal RM usage, such as stomp changes
 	commandOverride?: string
 	accumulateCommandOverride?: boolean
 
diff --git a/Engine/Extras/RoboMerge/v3/src/robo/branch-interfaces.ts b/Engine/Extras/RoboMerge/v3/src/robo/branch-interfaces.ts
--- a/Engine/Extras/RoboMerge/v3/src/robo/branch-interfaces.ts
+++ b/Engine/Extras/RoboMerge/v3/src/robo/branch-interfaces.ts
@@ -124,7 +124,7 @@ export interface TargetInfo {
 	author: string
 
 	targetWorkspaceForShelf?: string // Filled in during the reconsider in case of a createShelf nodeop request
-	sendNoShelfEmail: boolean // Used for internal use shelves, such as stomp changes
+	sendNoShelfNotification: boolean // Used for internal use shelves, such as stomp changes
 
 	forceStompChanges: boolean
 	additionalDescriptionText?: string
diff --git a/Engine/Extras/RoboMerge/v3/src/robo/branchdefs.ts b/Engine/Extras/RoboMerge/v3/src/robo/branchdefs.ts
--- a/Engine/Extras/RoboMerge/v3/src/robo/branchdefs.ts
+++ b/Engine/Extras/RoboMerge/v3/src/robo/branchdefs.ts
@@ -18,6 +18,9 @@ export interface BotConfig {
 	isDefaultBot: boolean
 	noStreamAliases: boolean
 	globalNotify: string[]
+	nagWhenBlocked: boolean | null
+	nagSchedule: number[] | null
+	triager: string | null
 	checkIntervalSecs: number
 	excludeAuthors: string[]
 	emailOnBlockage: boolean
@@ -44,6 +47,9 @@ const branchBasePrototype = {
 	forceFlowTo: [''],
 	defaultFlow: [''],
 	resolver: '' as string | null,
+	triager: '' as string | null,
+	nagWhenBlocked: false as boolean | null,
+	nagSchedule: [] as number[] | null,
 	aliases: [''],
 	badgeProject: '' as string | null,
 }
@@ -137,6 +143,10 @@ const edgeOptionFieldsPrototype = {
 	postOnlyToAdditionalChannel: false,
 
 	resolver: '',
+	triager: '',
+
+	nagSchedule: [],
+	nagWhenBlocked: true,
 
 	terminal: false, // changes go along terminal edges but no further
 
@@ -256,6 +266,9 @@ export class BranchDefs {
 			isDefaultBot: false,
 			noStreamAliases: false,
 			globalNotify: [],
+			triager: null,
+			nagSchedule: null,
+			nagWhenBlocked: null,
 			emailOnBlockage: true,
 			checkIntervalSecs: 30.0,
 			excludeAuthors: [],
diff --git a/Engine/Extras/RoboMerge/v3/src/robo/branchgraph.ts b/Engine/Extras/RoboMerge/v3/src/robo/branchgraph.ts
--- a/Engine/Extras/RoboMerge/v3/src/robo/branchgraph.ts
+++ b/Engine/Extras/RoboMerge/v3/src/robo/branchgraph.ts
@@ -54,8 +54,7 @@ if (botname === '__TEST__') {
 	private propagateBotPropertiesToNodes() {
 		const propertiesToPropagate: [keyof BotConfig, keyof NodeOptions, ConfigBlendMode | null][] = [
 			['excludeAuthors', 'excludeAuthors', 'override'],
-			['badgeUrlOverride', 'badgeUrlOverride', 'override'],
-
+			['badgeUrlOverride', 'badgeUrlOverride', 'override']
 		]
 
 		for (const [nameInBot, nameInNode, blendMode] of propertiesToPropagate) {
@@ -239,6 +238,9 @@ if (botname === '__TEST__') {
 			enabled: !options.disabled,
 			allowDeadend: !options.disallowDeadend,
 			resolver: options.resolver || null,
+			triager: options.triager || null,
+			nagSchedule: options.nagSchedule || null,
+			nagWhenBlocked: options.nagWhenBlocked || null,
 			convertIntegratesToEdits: (options.integrationMethod || this.config.defaultIntegrationMethod) === IntegrationMethod.CONVERT_TO_EDIT,
 			visibility: options.visibility || this.config.visibility,
 			get isMonitored() { return !!(this.bot && this.bot.isRunning) }
diff --git a/Engine/Extras/RoboMerge/v3/src/robo/conflict-interfaces.ts b/Engine/Extras/RoboMerge/v3/src/robo/conflict-interfaces.ts
--- a/Engine/Extras/RoboMerge/v3/src/robo/conflict-interfaces.ts
+++ b/Engine/Extras/RoboMerge/v3/src/robo/conflict-interfaces.ts
@@ -23,7 +23,9 @@ export interface PersistentConflict {
 	kind: FailureKind
 	time: Date
 
-	nagged: boolean
+	nagCount: number
+	lastNagTime?: Date
+
 	ugsIssue: number
 
 	resolution?: Resolution
diff --git a/Engine/Extras/RoboMerge/v3/src/robo/conflicts.ts b/Engine/Extras/RoboMerge/v3/src/robo/conflicts.ts
--- a/Engine/Extras/RoboMerge/v3/src/robo/conflicts.ts
+++ b/Engine/Extras/RoboMerge/v3/src/robo/conflicts.ts
@@ -35,7 +35,7 @@ function makePersistentConflict(blockage: Blockage): PersistentConflict {
 		owner: blockage.owner,
 		kind: blockage.failure.kind,
 		time: blockage.time,
-		nagged: false,
+		nagCount: 0,
 		ugsIssue: -1
 	}
 
@@ -104,6 +104,10 @@ export class Conflicts {
 		if (conflictsToLoad) {
 			for (const conflict of conflictsToLoad) {
 				conflict.time = new Date(conflict.time)
+				if (conflict.lastNagTime) {
+					conflict.lastNagTime = new Date(conflict.lastNagTime)
+				}
+				conflict.nagCount = conflict.nagCount || 0 // back compat for nag changes
 				this.conflicts.push(conflict)
 			}
 		}
@@ -184,7 +188,8 @@ export class Conflicts {
 					replacement.acknowledger = acknowledger
 					replacement.acknowledgedAt = conflict.acknowledgedAt
 				}
-				replacement.nagged = conflict.nagged
+				replacement.nagCount = conflict.nagCount
+				replacement.lastNagTime = conflict.lastNagTime
 				replacement.ugsIssue = conflict.ugsIssue
 
 				this.conflicts[index] = replacement
diff --git a/Engine/Extras/RoboMerge/v3/src/robo/edgebot.ts b/Engine/Extras/RoboMerge/v3/src/robo/edgebot.ts
--- a/Engine/Extras/RoboMerge/v3/src/robo/edgebot.ts
+++ b/Engine/Extras/RoboMerge/v3/src/robo/edgebot.ts
@@ -536,7 +536,7 @@ class EdgeBotImpl extends PerforceStatefulBot {
 			case 'no_files': {
 				const msg = `Change ${to_integrate} was not necessary in ${this.targetBranch.name}`
 				const event: AlreadyIntegrated = {change: info, action: target}
-				this.sourceNode.onAlreadyIntegrated(msg, event)
+				this.sourceNode.onAlreadyIntegrated(event)
 
 				// integration not necessary
 				this.edgeBotLogger.info(msg)
@@ -576,7 +576,7 @@ class EdgeBotImpl extends PerforceStatefulBot {
 		}
 
 		// Send to source node to facilitate notification handling
-		if (this.sourceNode.handleMergeFailure(failure, pending)) {
+		if (await this.sourceNode.handleMergeFailure(failure, pending)) {
 			this.gate.onBlockage()
 		}
 		return new EdgeIntegrationDetails('error', errors.join('\n'))
@@ -623,8 +623,8 @@ class EdgeBotImpl extends PerforceStatefulBot {
 			// the user requested manual merge
 			await this.shelveChangelist(pending)
 
-			if (!pending.change.sendNoShelfEmail) {
-				this.sourceNode.emailShelfRequester(pending)
+			if (!pending.change.sendNoShelfNotification) {
+				await this.sourceNode.notifyShelfRequester(pending)
 			}
 			
 			return new EdgeIntegrationDetails('shelved', pending.newCl)
@@ -660,14 +660,22 @@ class EdgeBotImpl extends PerforceStatefulBot {
 			// or lock a test stream and try to commit
 			// then look for error message in result.message
 			if (!this.options.approval) {
-				failure = { kind: 'Commit failure', description: result.message }
+
+				let summary: string | undefined
+				const match = result.message.match(/.*STDERR:([^]*)STDOUT:/)
+				if (match)
+				{
+					summary = match[1].trim()
+				}
+
+				failure = { kind: 'Commit failure', description: result.message, summary }
 			}
 		}
 
 		// no failure set means the change needs approval
 		if (failure) {
 			await this.handlePostIntegrationFailure(failure, pending)
-			this.sourceNode.handleMergeFailure(failure, pending, true)
+			await this.sourceNode.handleMergeFailure(failure, pending, true)
 		}
 		else {
 			const approval = this.options.approval!
@@ -820,7 +828,7 @@ class EdgeBotImpl extends PerforceStatefulBot {
 
 			pending.newCl = -1
 
-			this.sourceNode._sendError(recipients, `Error merging ${pending.change.source_cl} to ${this.targetBranch.name}`, final_desc)
+			this.sourceNode.sendErrorEmail(recipients, `Error merging ${pending.change.source_cl} to ${this.targetBranch.name}`, final_desc)
 			return
 		}
 
diff --git a/Engine/Extras/RoboMerge/v3/src/robo/graphbot.ts b/Engine/Extras/RoboMerge/v3/src/robo/graphbot.ts
--- a/Engine/Extras/RoboMerge/v3/src/robo/graphbot.ts
+++ b/Engine/Extras/RoboMerge/v3/src/robo/graphbot.ts
@@ -15,7 +15,8 @@ import { PersistentConflict } from './conflict-interfaces';
 import { BotEventHandler, BotEventTriggers } from './events';
 import { GraphInterface } from './graph-interface';
 import { NodeBot } from './nodebot';
-import { bindBotNotifications, BotNotifications, NOTIFICATIONS_PERSISTENCE_KEY, postMessageToChannel, postToRobomergeAlerts } from './notifications';
+import { bindBotNotifications, BotNotifications, NOTIFICATIONS_PERSISTENCE_KEY } from './notifications';
+import { postMessageToChannel, postToRobomergeAlerts, SlackMessages } from './notifications';
 import { roboAnalytics } from './roboanalytics';
 import { BlockageNodeOpUrls, OperationUrlHelper } from './roboserver';
 import { Settings } from './settings';
@@ -36,6 +37,7 @@ export class GraphBot implements GraphInterface, BotEventHandler {
 	autoUpdater: AutoBranchUpdater | null
 
 	private botLogger: ContextualLogger;
+	private slackMessages?: SlackMessages
 
 	// separate off into class that only exists while bots are running?
 	private eventTriggers?: BotEventTriggers;
@@ -122,14 +124,14 @@ export class GraphBot implements GraphInterface, BotEventHandler {
 			}
 		}
 
-		bindBotNotifications(this.eventTriggers, slackChannelOverrides, this.settings.getContext(NOTIFICATIONS_PERSISTENCE_KEY), blockageUrlGenerator, this.externalUrl, this.botLogger)
+		this.slackMessages = bindBotNotifications(this.eventTriggers, slackChannelOverrides, this.settings.getContext(NOTIFICATIONS_PERSISTENCE_KEY), blockageUrlGenerator, this.externalUrl, this.botLogger)
 		bindBadgeHandler(this.eventTriggers, this.branchGraph, this.externalUrl, this.botLogger)
 
 		let hasConflicts = false
 		for (const branch of this.branchGraph.branches) {
 			if (branch.enabled) {
 				const persistence = this.settings.getContext(branch.upperName)
-				branch.bot = new NodeBot(branch, this.mailer, this.externalUrl, this.eventTriggers, persistence, ubergraph,
+				branch.bot = new NodeBot(branch, this.mailer, this.slackMessages, this.externalUrl, this.eventTriggers, persistence, ubergraph,
 					async () => {
 						const errPair = await this.handleRequestedIntegrationsForAllNodes()
 						if (errPair) {
@@ -392,7 +394,7 @@ export class GraphBot implements GraphInterface, BotEventHandler {
 		}
 
 		let botNotify = new BotNotifications(this.branchGraph.botname, this.branchGraph.config.slackChannel,
-			this.settings.getContext(NOTIFICATIONS_PERSISTENCE_KEY), this.externalUrl, fakeHelper, this.botLogger)
+			this.externalUrl, fakeHelper, this.botLogger, this.slackMessages)
 
 		return botNotify.sendTestMessage(username)
 	}
diff --git a/Engine/Extras/RoboMerge/v3/src/robo/nodebot.ts b/Engine/Extras/RoboMerge/v3/src/robo/nodebot.ts
--- a/Engine/Extras/RoboMerge/v3/src/robo/nodebot.ts
+++ b/Engine/Extras/RoboMerge/v3/src/robo/nodebot.ts
@@ -6,16 +6,18 @@ import { ContextualLogger } from '../common/logger';
 import { Mailer, MailParams, Recipients } from '../common/mailer';
 import { Change, ConflictedResolveNFile, PerforceContext, RoboWorkspace, coercePerforceWorkspace } from '../common/perforce';
 import { IPCControls, NodeBotInterface, QueuedChange, ReconsiderArgs } from './bot-interfaces';
-import { ApprovalOptions } from './branchdefs'
+import { ApprovalOptions, BotConfig, EdgeOptions } from './branchdefs'
 import { BlockagePauseInfo, BranchStatus } from './status-types';
-import { AlreadyIntegrated, Blockage, Branch, BranchArg, BranchGraphInterface, ChangeInfo, EndIntegratingToGateEvent, Failure } from './branch-interfaces';
+import { AlreadyIntegrated, Blockage, Branch, BranchArg, BranchGraphInterface, ChangeInfo, EditableBranch, EndIntegratingToGateEvent, Failure } from './branch-interfaces';
 import { MergeAction, OperationResult, PendingChange, resolveBranchArg, StompedRevision, StompVerification, StompVerificationFile }  from './branch-interfaces';
 import { Conflicts } from './conflicts';
 import { EdgeBot, EdgeMergeResults } from './edgebot';
 import { BotEventTriggers } from './events';
+import { SlackMessages } from './notifications';
 import { PerforceStatefulBot } from './perforce-stateful-bot';
 import { BlockageNodeOpUrls, OperationUrlHelper } from './roboserver';
 import { Context } from './settings';
+import { SlackMessage } from './slack';
 import { PauseState } from './state-interfaces';
 import { newTickJournal, TickJournal } from './tick-journal';
 import { computeTargets, parseDescriptionLines, processOtherBotTargets, getIntegrationOwner, getNodeBotFullName, getNodeBotFullNameForLogging } from './targets';
@@ -25,8 +27,8 @@ import { GraphAPI } from '../new/graph';
  * Bot monitoring a single stream
  **********************************/
 
-
-const NAG_EMAIL_MIN_TIME_SECONDS = 60 * 60 * 1000
+const DEFAULT_NAG_SCHEDULE = [10, 20, 30, 60]
+const NAG_EMAIL_MIN_TIME_MINUTES = 60
 const NAG_EMAIL_MIN_TIME_DESCRIPTION = 'an hour'
 const SYNTAX_ERROR_PAUSE_TIMEOUT_SECONDS = 10 * 60
 
@@ -62,6 +64,7 @@ export class NodeBot extends PerforceStatefulBot implements NodeBotInterface {
 	private headCL = -1
 
 	private readonly mailer: Mailer
+	private readonly slackMessages?: SlackMessages
 	protected readonly p4: PerforceContext
 
 	private readonly conflicts: Conflicts
@@ -82,7 +85,8 @@ export class NodeBot extends PerforceStatefulBot implements NodeBotInterface {
 
 	constructor(
 		branchDef: Branch, 
-		mailer: Mailer, 
+		mailer: Mailer,
+		slackMessages: SlackMessages | undefined,
 		externalUrl: string, 
 		private readonly eventTriggers: BotEventTriggers,
 		settings: Context,
@@ -96,6 +100,7 @@ export class NodeBot extends PerforceStatefulBot implements NodeBotInterface {
 		this.branchGraph = branchDef.parent
 		this.graphBotName = this.branchGraph.botname
 		this.nodeBotLogger = new ContextualLogger(this.fullNameForLogging)
+		this.slackMessages = slackMessages
 
 		this.externalRobomergeUrl = externalUrl
 
@@ -284,11 +289,14 @@ export class NodeBot extends PerforceStatefulBot implements NodeBotInterface {
 			return false
 		}
 
+		if (this.conflicts.getConflicts().length > 0)
+		{
+			this.sendNagNotifications()
+		}
+
 		// see if our flow is blocked
 		if (this.isBlocked) {
 			this.nodeBotLogger.verbose('tick() - This bot is blocked')
-
-			this.sendNagEmails()
 			await _nextTick()
 			return false
 		}
@@ -356,7 +364,7 @@ export class NodeBot extends PerforceStatefulBot implements NodeBotInterface {
 
 			if (!specifiedTargetBranch) {
 				// We somehow have a specified branch from a reconsider request, but it doesn't exist in the branch map. Report the error
-				this._sendError(
+				this.sendErrorEmail(
 					new Recipients(fromQueue.who),
 					`Invalid Reconsider Request -- Branch "${fromQueue.targetBranchName}" does not exist`,
 					`Manually queued change ${fromQueue.cl} on ${this.fullName}, requested by ${fromQueue.who}, ` + 
@@ -366,7 +374,7 @@ export class NodeBot extends PerforceStatefulBot implements NodeBotInterface {
 			}
 			else if (specifiedTargetBranch === this.branch) {
 				// We somehow have targeted our own branch with a merge. Report the issue
-				this._sendError(
+				this.sendErrorEmail(
 					new Recipients(fromQueue.who),
 					`Invalid Reconsider Request -- Trying to merge "${fromQueue.targetBranchName}" into itself`,
 					`Manually queued change ${fromQueue.cl} on ${this.fullName}, requested by ${fromQueue.who}, ` +
@@ -406,8 +414,8 @@ export class NodeBot extends PerforceStatefulBot implements NodeBotInterface {
 		if (additionalFlags.indexOf('createashelf') !== -1) {
 			change.forceCreateAShelf = true
 		}
-		if (additionalFlags.indexOf('sendnoshelfemail') != -1) {
-			change.sendNoShelfEmail = true
+		if (additionalFlags.indexOf('sendnoshelfnotification') != -1) {
+			change.sendNoShelfNotification = true
 		}
 
 		// Stomp Changes support
@@ -593,7 +601,7 @@ export class NodeBot extends PerforceStatefulBot implements NodeBotInterface {
 		// Set a variety of flags on the change before processing
 		blockageChange.isUserRequest = true
 		blockageChange.forceCreateAShelf = true // We don't want to submit this
-		blockageChange.sendNoShelfEmail = true // Don't send any shelf creation emails for this
+		blockageChange.sendNoShelfNotification = true // Don't send any shelf creation notifications for this
 
 		// Process and attempt the integration
 		const edgeMap = new Map<string, EdgeBot>([[edge.targetBranch.upperName, edge]])
@@ -1203,8 +1211,7 @@ export class NodeBot extends PerforceStatefulBot implements NodeBotInterface {
 		}
 	}
 
-	onAlreadyIntegrated(msg: string, event: AlreadyIntegrated) {
-		this._emailNoActionIfRequested(event.change, msg)
+	onAlreadyIntegrated(event: AlreadyIntegrated) {
 		this.conflicts.onAlreadyIntegrated(event)
 	}
 
@@ -1218,6 +1225,11 @@ export class NodeBot extends PerforceStatefulBot implements NodeBotInterface {
 		}
 	}
 
+	findEmail = (user: string) => {
+		this._log_action(`Getting email for user ${user}`);
+		return this.p4.getEmail(user);
+	}
+
 	// Listener method to email a blockage owner when the branch encounters said blockage,
 	// depending on NodeBot configuration
 	private emailBlockageCulprit(blockage: Blockage) {
@@ -1238,11 +1250,11 @@ export class NodeBot extends PerforceStatefulBot implements NodeBotInterface {
 		this._sendBlockageEmail(blockage)
 	}
 
-	emailShelfRequester(pendingChange: PendingChange) {
+	async notifyShelfRequester(pendingChange: PendingChange) {
 		// Get branch referenced by blockage object
 		const shelfBranch : Branch = pendingChange.change.branch
 
-		// Only email for our branch, if the branch is configured to receive emails
+		// Only notify for our branch, if the branch is configured to receive notifications
 		if (this.branch.name !== shelfBranch.name) {
 			return
 		}
@@ -1250,18 +1262,33 @@ export class NodeBot extends PerforceStatefulBot implements NodeBotInterface {
 		const owner = pendingChange.change.owner
 		// Should never happen
 		if (!owner) {
-			this.nodeBotLogger.warn(`Unable to send shelf creation email for source CL ${pendingChange.change.source_cl} in branch ${shelfBranch} because there was no owner.`)
+			this.nodeBotLogger.warn(`Unable to send shelf creation notification for source CL ${pendingChange.change.source_cl} in branch ${shelfBranch} because there was no owner.`)
 			return
 		}
 
-		this.nodeBotLogger.info(`Sending reconsider shelf creation (shelf CL ${pendingChange.newCl}) email to ${pendingChange.change.owner} for source CL ${pendingChange.change.source_cl}`)
+		this.nodeBotLogger.info(`Sending reconsider shelf creation (shelf CL ${pendingChange.newCl}) notification to ${pendingChange.change.owner} for source CL ${pendingChange.change.source_cl}`)
 
-		this._sendGenericEmail(
-			new Recipients(owner),
-			`Robomerge created Shelf CL ${pendingChange.newCl} in ${pendingChange.change.targetWorkspaceForShelf}`,
-			`Robomerge has created shelf CL ${pendingChange.newCl} in workspace ${pendingChange.change.targetWorkspaceForShelf} ` +
-				`for merging source CL ${pendingChange.change.source_cl} (${this.fullName}).`,
-			`${pendingChange.change.source_cl}:\n${pendingChange.change.description}`)
+		if (this.slackMessages) {
+			let dm: SlackMessage = {
+				text: `Robomerge has created shelf CL ${pendingChange.newCl} in workspace ${pendingChange.change.targetWorkspaceForShelf} ` +
+				`for merging source CL ${pendingChange.change.source_cl} (${this.fullName}).\n` +
+				`\`\`\`${pendingChange.change.description}\`\`\``,
+				channel: '',
+				mrkdwn: true
+			}
+			
+			let emailAddress = await this.findEmail(owner)
+			this.slackMessages.postDM(emailAddress, pendingChange.newCl, this.branch, dm)
+		}
+		else
+	 	{
+			this._sendGenericEmail(
+				new Recipients(owner),
+				`Robomerge created Shelf CL ${pendingChange.newCl} in ${pendingChange.change.targetWorkspaceForShelf}`,
+				`Robomerge has created shelf CL ${pendingChange.newCl} in workspace ${pendingChange.change.targetWorkspaceForShelf} ` +
+					`for merging source CL ${pendingChange.change.source_cl} (${this.fullName}).`,
+				`${pendingChange.change.source_cl}:\n${pendingChange.change.description}`)
+		}
 	}
 
 	private async _processListOfChanges(availableEdges: Map<string, EdgeBot>, allChanges: Change[], maxChangesToProcess: number) {
@@ -1472,18 +1499,12 @@ export class NodeBot extends PerforceStatefulBot implements NodeBotInterface {
 		if (!info.targets) {
 			const nothingToDoMesasage = `Nothing to do for ${info.cl}`
 			this.nodeBotLogger.silly(nothingToDoMesasage)
-			this._emailNoActionIfRequested(info, nothingToDoMesasage);
 			return { errors:  [ nothingToDoMesasage ] }
 		}
 
 		return { info, errors: [] }
 	}
 
-	findEmail = (user: string) => {
-		this._log_action(`Getting email for user ${user}`);
-		return this.p4.getEmail(user);
-	}
-
 	private createEmailLinkButton(buttonText: string, buttonLink: string, buttonDescription: string, textColor: string, buttonColor: string): string {
 		// Keeping it simple for now
 		const borderColor = textColor
@@ -1563,15 +1584,7 @@ export class NodeBot extends PerforceStatefulBot implements NodeBotInterface {
 		this.mailer.sendEmail('new-blockage', recipients, subject, params, this.findEmail)
 	}
 
-	_emailNoActionIfRequested(info: ChangeInfo, _msg: string) {
-
-		if (info.userRequest) {
-			// this._sendEmail(new Recipients(info.owner!), msg, `Just an FYI that RoboMerge (${this.botname}) did not perform an integration for this changelist`, info.description);
-		}
-	}
-
-	// TEMPORARY: Removing private from _sendError while EdgeBot work continues
-	_sendError(recipients: Recipients, subject: string, message: string, errorPreface='RoboMerge needs your help:'): void {
+	sendErrorEmail(recipients: Recipients, subject: string, message: string, errorPreface='RoboMerge needs your help:'): void {
 		// log to STDERR
 		this.nodeBotLogger.error(`${subject}\n${message}\n\n`);
 		this._sendGenericEmail(recipients, subject, errorPreface, message);
@@ -1630,7 +1643,7 @@ export class NodeBot extends PerforceStatefulBot implements NodeBotInterface {
 		// By default, we should perform any submit that comes from a change.
 		// If the change is setup to be forcibly shelved, honor it
 		let forceCreateAShelf = !!change.forceCreateAShelf
-		let sendNoShelfEmail = !!change.sendNoShelfEmail
+		let sendNoShelfNotification = !!change.sendNoShelfNotification
 
 		let forceStompChanges = !!change.forceStompChanges
 		let additionalDescriptionText = change.additionalDescriptionText
@@ -1643,7 +1656,7 @@ export class NodeBot extends PerforceStatefulBot implements NodeBotInterface {
 			source: parsedLines.source,
 			description,
 			propagatingNullMerge: parsedLines.propagatingNullMerge,
-			forceCreateAShelf, sendNoShelfEmail,
+			forceCreateAShelf, sendNoShelfNotification,
 
 			forceStompChanges, additionalDescriptionText,
 			overriddenCommand: commandOverride,
@@ -1833,7 +1846,7 @@ export class NodeBot extends PerforceStatefulBot implements NodeBotInterface {
 		this.eventTriggers.reportBlockage(blockage, true)
 	}
 
-	private handleSyntaxError(originalCommitMessage: string, change: ChangeInfo) {
+	private async handleSyntaxError(originalCommitMessage: string, change: ChangeInfo) {
 		const shortMessage = `Syntax error in CL ${change.cl}`
 		let message = `CL ${change.cl} in ${this.branch.rootPath} contained one or more RoboMerge syntax errors (see below).\n\n`
 		for (const err of change.errors!) {
@@ -1844,10 +1857,23 @@ export class NodeBot extends PerforceStatefulBot implements NodeBotInterface {
 
 		if (change.userRequest) {
 			this.nodeBotLogger.info(shortMessage + ` (reconsider triggered by ${change.owner}):\n${message}.`)
-			this._sendError(new Recipients(owner), shortMessage, message + '\nFull CL description:\n\n' + originalCommitMessage)
+			if (this.slackMessages) {
+
+				let dm: SlackMessage = {
+					text: message,
+					channel: '',
+					mrkdwn: true
+				}
+				
+				let emailAddress = await this.findEmail(owner)
+				this.slackMessages.postDM(emailAddress, change.cl, this.branch, dm)				
+			}
+			else {
+				this.sendErrorEmail(new Recipients(owner), shortMessage, message + '\nFull CL description:\n\n' + originalCommitMessage)
+			}
 		}
 		else {
-			// see if we've already sent an email for this one
+			// see if we've already sent a notification for this one
 			const existingBlockage = this.conflicts.find(null, change.source_cl)
 			const blockageOccurred = existingBlockage ? existingBlockage.time : new Date
 
@@ -1864,7 +1890,7 @@ export class NodeBot extends PerforceStatefulBot implements NodeBotInterface {
 	}
 
 	/** return true if this is a new flow blockage (excludes reconsiders) */
-	handleMergeFailure(failure: Failure, pending: PendingChange, postIntegrate?: boolean) {
+	async handleMergeFailure(failure: Failure, pending: PendingChange, postIntegrate?: boolean) {
 		const targetBranch = pending.action.branch
 		const shortMessage = `${failure.kind} while merging CL ${pending.change.cl} to ${targetBranch.name}`
 
@@ -1877,8 +1903,28 @@ export class NodeBot extends PerforceStatefulBot implements NodeBotInterface {
 				message += `Shelved in workspace ${pending.change.targetWorkspaceForShelf}\n`
 			}
 
-			message += '\n' + failure.summary 
-			this._sendError(new Recipients(owner), shortMessage, message)
+			if (failure.summary) {
+				message += '\n' + failure.summary 
+			}
+			else
+			{
+				message += '\n' + failure.description
+			}
+
+			if (this.slackMessages) {
+
+				let dm: SlackMessage = {
+					text: message,
+					channel: '',
+					mrkdwn: true
+				}
+				
+				let emailAddress = await this.findEmail(owner)
+				this.slackMessages.postDM(emailAddress, pending.newCl, this.branch, dm)
+			}
+			else {
+				this.sendErrorEmail(new Recipients(owner), shortMessage, message)
+			}
 			return false
 		}
 
@@ -1909,22 +1955,102 @@ export class NodeBot extends PerforceStatefulBot implements NodeBotInterface {
 		return true
 	}
 
-	private sendNagEmails() {
+	private async sendNagNotifications() {
 		let naggedSomeone = false
 		const now = Date.now()
 		for (const conflict of this.conflicts.getConflicts()) {
-			if (!conflict.nagged) {
-				const ageSeconds = (now - conflict.time.getTime()) / 1000
-				if (ageSeconds > NAG_EMAIL_MIN_TIME_SECONDS) {
-					conflict.nagged = true
-					naggedSomeone = true
-
-					// send a one-time only additional nag email just to owner
-					this.nodeBotLogger.info(`Sending nag email to ${conflict.owner}`)
-					this._sendError(new Recipients(conflict.owner),
-						`${this.graphBotName} RoboMerge blocked for more than ${NAG_EMAIL_MIN_TIME_DESCRIPTION}`,
-						`Please check #robomerge-${this.graphBotName.toLowerCase()} on Slack for more details`
-					)
+
+			const getNagProperty = (property:string) => {
+						
+				const configKey = property as keyof BotConfig
+				const edgeKey   = property as keyof EdgeOptions
+				const branchKey = property as keyof EditableBranch
+
+				let prop: any | null = this.branchGraph.config[configKey]
+				if (conflict.targetBranchName) {
+					const targetBranch = this.branchGraph.getBranch(conflict.targetBranchName)
+					if (targetBranch)
+					{
+						prop = targetBranch[branchKey] || prop
+					}
+
+					const edgeProperties = this.branch.edgeProperties.get(conflict.targetBranchName)
+					if (edgeProperties && edgeProperties[edgeKey]) {
+						prop = edgeProperties[edgeKey] || prop
+					}
+				}
+				else
+				{
+					prop = this.branch[branchKey] || prop
+				}
+				return prop
+			}
+			
+			const nagWhenBlocked = getNagProperty("nagWhenBlocked")
+			
+			if (nagWhenBlocked) {
+				const ageMinutes = ((now - conflict.time.getTime()) / 1000) / 60
+
+				if (this.slackMessages) {
+
+					const nagSchedule = getNagProperty("nagSchedule") || DEFAULT_NAG_SCHEDULE
+
+					const nagDelay = conflict.lastNagTime ? ((now - conflict.lastNagTime.getTime()) / 1000) / 60 : ageMinutes
+					const nagWait = nagSchedule[Math.min(conflict.nagCount, nagSchedule.length - 1)]
+
+					if (nagDelay > nagWait) {
+						conflict.nagCount++ 
+						conflict.lastNagTime = new Date(now)
+						naggedSomeone = true
+
+						let timeDesc
+						if (ageMinutes < 60) {
+							timeDesc = `${Math.floor(ageMinutes)} minutes`
+						}
+						else {
+							const hours = Math.floor(ageMinutes / 60)
+							if (hours == 1) {
+								timeDesc = "1 hour"
+							}
+							else {
+								timeDesc = `${hours} hours`
+							}
+						}
+				
+						let triager = getNagProperty("triager")
+
+						if (triager && !triager.startsWith('@'))
+						{
+							const emailAddress = await this.findEmail(triager)
+							if (emailAddress)
+							{
+								triager = `<@${await this.slackMessages.getSlackUser(emailAddress)}>`
+							}
+						}
+
+						const nagMessage: SlackMessage = {
+							text: `${triager} RoboMerge blocked for more than ${timeDesc}`,
+							channel: this.branchGraph.config.slackChannel,
+							mrkdwn: true
+						}
+
+						this.nodeBotLogger.info(`Sending nag notification to ${conflict.owner} after ${Math.floor(ageMinutes)} minutes`)
+						const branchName =  (conflict.kind === "Syntax error" ? conflict.kind : conflict.targetBranchName || conflict.blockedBranchName)
+						this.slackMessages.postReply(conflict.cl, branchName, nagMessage)
+					}
+				}
+				else if (conflict.nagCount == 0) // for now we'll maintain the email only once behavior
+				{
+					if (ageMinutes > NAG_EMAIL_MIN_TIME_MINUTES) {
+
+						conflict.nagCount = 1
+						naggedSomeone = true
+
+						this.sendErrorEmail(new Recipients(conflict.owner),
+							`${this.graphBotName} RoboMerge blocked for more than ${NAG_EMAIL_MIN_TIME_DESCRIPTION}`,
+							`Please check robomerge page for more details`
+						)
+					}
 				}
 			}
 		}
diff --git a/Engine/Extras/RoboMerge/v3/src/robo/notifications.ts b/Engine/Extras/RoboMerge/v3/src/robo/notifications.ts
--- a/Engine/Extras/RoboMerge/v3/src/robo/notifications.ts
+++ b/Engine/Extras/RoboMerge/v3/src/robo/notifications.ts
@@ -18,9 +18,6 @@ import { DummySlackApp } from '../common/dummyslackserver'
 
 const notificationsStartupLogger = new ContextualLogger('Notifications Startup')
 
-// var to enable DMs on blockage
-const DIRECT_MESSAGING_ENABLED = true
-
 export const NOTIFICATIONS_PERSISTENCE_KEY = 'notifications'
 export const SLACK_MESSAGES_PERSISTENCE_KEY = 'slackMessages'
 const CHANGELIST_FIELD_TITLE = 'Change'
@@ -71,7 +68,7 @@ export function isUserAKnownBot(user: string) {
 
 
 export async function postToRobomergeAlerts(message: string) {
-	// 'CBGFJQEN6' is #robomerge_alerts
+	// 'C044ZUSS71S' is #dt-robomerge-alert-ext
 	if (!args.devMode) {
 		return postMessageToChannel(message, 'C044ZUSS71S')
 	}
@@ -212,7 +209,7 @@ type SlackConflictMessage = {
 type PersistedSlackMessages = {[key: string]: SlackConflictMessage}
 
 /** Wrapper around Slack, keeping track of messages keyed on target branch and source CL */
-class SlackMessages {
+export class SlackMessages {
 	private readonly smLogger: ContextualLogger
 	constructor(private slack: Slack, private persistence: Context, parentLogger: ContextualLogger) {
 		this.smLogger = parentLogger.createChild('SlackMsgs')
@@ -266,7 +263,24 @@ class SlackMessages {
 		}
 	}
 
-	async postDM(emailAddress: string, sourceCl: number, branchArg: BranchArg, dm: SlackMessage, persistMessage = true) {
+	async postReply(sourceCl: number, branchArg: BranchArg, message: SlackMessage) {
+		const findResult = this.find(sourceCl, branchArg, message.channel)
+
+		if (findResult.messageRecord) {
+			try {
+				await this.slack.reply(findResult.messageRecord.timestamp, message)
+			}
+			catch (err) {
+				console.error('Error updating message in Slack! ' + err.toString())
+				return
+			}
+		}
+		else {
+			this.smLogger.error(`Failed to find message record for ${sourceCl}:${branchArg}:${message.channel}`)
+		}
+	}
+
+	async postDM(emailAddress: string|null, sourceCl: number, branchArg: BranchArg, dm: SlackMessage, persistMessage = true) {
 		// The Slack API requires a user ID to open a direct message with users.
 		// The most consistent way to do this is getting their email address out of P4.
 		if (!emailAddress) {
@@ -274,10 +288,9 @@ class SlackMessages {
 			return
 		}
 
-		// With their email address, we can get their user ID via Slack API
 		let userId : string
 		try {
-			userId = (await this.slack.lookupUserIdByEmail(emailAddress.toLowerCase()))
+			userId = (await this.getSlackUser(emailAddress))
 		} catch (err) {
 			this.smLogger.printException(err, `Failed to get user ID for Slack DM, given email address "${emailAddress}" for CL ${sourceCl}`)
 			return
@@ -341,6 +354,21 @@ class SlackMessages {
 		.catch(err => console.error('Error posting non-conflict to Slack! ' + err.toString()))
 	}
 
+	async addUserToChannel(emailAddress: string, channel: string)
+	{
+		const user = await this.getSlackUser(emailAddress)
+		const result = await this.slack.addUserToChannel(user, channel)
+		if (!result.ok && result.error !== "already_in_channel") {
+			this.smLogger.error(`Error inviting ${emailAddress} (${user}) to channel ${channel}: ${result.error}`)
+		}
+	}
+
+	// With their email address, we can get their user ID via Slack API
+	getSlackUser(emailAddress: string) {
+		
+		return this.slack.lookupUserIdByEmail(emailAddress.toLowerCase())
+	}
+
 	private doHouseKeeping(messages: PersistedSlackMessages) {
 		const nowTicks = Date.now() / 1000.
 
@@ -373,29 +401,25 @@ export class BotNotifications implements BotEventHandler {
 	private readonly botNotificationsLogger : ContextualLogger
 	slackChannel: string;
 
-	constructor(private botname: string, slackChannel: string, persistence: Context, externalUrl: string, 
+	constructor(private botname: string, slackChannel: string, externalUrl: string, 
 		blockageUrlGenerator: NodeOpUrlGenerator, parentLogger: ContextualLogger,
-		slackChannelOverrides?: [Branch, Branch, string, boolean][]) {
+		slackMessages?: SlackMessages, slackChannelOverrides?: [Branch, Branch, string, boolean][]) {
 		this.botNotificationsLogger = parentLogger.createChild('Notifications')
 		// Hacky way to dynamically change the URL for notifications
 		this.externalRobomergeUrl = externalUrl
 		this.slackChannel = slackChannel
+		this.slackMessages = slackMessages
 		this.blockageUrlGenerator = blockageUrlGenerator
 
-		const botToken = args.devMode && SLACK_DEV_DUMMY_TOKEN || SLACK_TOKENS.bot
-		if (botToken && slackChannel) {
-			this.botNotificationsLogger.info('Enabling Slack messages for ' + botname)
-			this.slackMessages = new SlackMessages(new Slack({id: slackChannel, botToken}, args.slackDomain), persistence, this.botNotificationsLogger)
-			if (slackChannelOverrides) {
-				this.additionalBlockChannelIds = new Map(slackChannelOverrides
-					.filter(
-						([_1, _2, channel, _3]) => channel !== slackChannel
-					)
-					.map(
-						([source, target, channel, postOnlyToChannel]) => [`${source.upperName}|${target.upperName}`, [channel, postOnlyToChannel]]
-					)
+		if (slackMessages && slackChannelOverrides) {
+			this.additionalBlockChannelIds = new Map(slackChannelOverrides
+				.filter(
+					([_1, _2, channel, _3]) => channel !== slackChannel
 				)
-			}
+				.map(
+					([source, target, channel, postOnlyToChannel]) => [`${source.upperName}|${target.upperName}`, [channel, postOnlyToChannel]]
+				)
+			)
 		}
 	}
 
@@ -428,9 +452,8 @@ export class BotNotifications implements BotEventHandler {
 		}
 		const issue = blockage.failure.kind.toLowerCase()
 
-		// If we get the user's email, we shouldn't ALSO ping them in the channel.
 		const userEmail = await blockage.ownerEmail
-		const channelPing = userEmail ? blockage.owner : `@${blockage.owner}`
+		const channelPing = `@${blockage.owner}`
 
 		const isBotUser = isUserAKnownBot(blockage.owner)
 		const text =
@@ -458,21 +481,43 @@ export class BotNotifications implements BotEventHandler {
 			if (additionalChannelInfo) {
 				const [sideChannel, postOnlyToSideChannel] = additionalChannelInfo
 				if (!postOnlyToSideChannel) {
+					if (userEmail) {
+						this.slackMessages.addUserToChannel(userEmail, message.channel)
+					}
 					this.slackMessages.postOrUpdate(changeInfo.source_cl, targetBranch.name, message)
 				}
+				if (userEmail) {
+					this.slackMessages.addUserToChannel(userEmail, sideChannel)
+				}
 				this.slackMessages.postOrUpdate(changeInfo.source_cl, targetBranch.name, { ...message, channel: sideChannel })
 			}
 			else {
+				if (userEmail) {
+					this.slackMessages.addUserToChannel(userEmail, message.channel)
+				}
 				this.slackMessages.postOrUpdate(changeInfo.source_cl, targetBranch.name, message)
 			}
 		}
 		else {
+
+			if (userEmail) {
+				this.slackMessages.addUserToChannel(userEmail, message.channel)
+			}
+
 			// separate message for syntax errors (other blockages have a target branch)
 			this.slackMessages.postOrUpdate(changeInfo.source_cl, blockage.failure.kind, message)
+
+			let syntaxErrorMessage: SlackMessage = {
+				text: blockage.failure.description,
+				channel: message.channel,
+				mrkdwn: false
+			}
+			
+			this.slackMessages.postReply(changeInfo.source_cl, blockage.failure.kind, syntaxErrorMessage)
 		}
 
 		// Post message to owner in DM
-		if (DIRECT_MESSAGING_ENABLED && !isBotUser && targetBranch && userEmail) {
+		if (!isBotUser && targetBranch && userEmail) {
 			let dm: SlackMessage
 			if (blockage.approval) {
 				dm = {
@@ -573,33 +618,39 @@ export class BotNotifications implements BotEventHandler {
 	}
 
 	sendTestMessage(username : string) {
-		let text = `${username}, please resolve the following test message:\n(source CL: 0, conflict CL: 1, shelf CL: 2)`
-
-		// This doesn't need to be a full Blockage -- just enough to generate required info for message
-		const testBlockage: Blockage = {
-			action: {
-				branch: {
-					name: "TARGET_BRANCH",
-				} as Branch,
-			} as MergeAction,
-			failure: {
-				kind: "Merge conflict",
-				description: ""
-			}
-		} as Blockage
-
-		const messageOpts = this.makeSlackDirectMessage(text, 0, 1, "TARGETBRANCH",
-			NodeBot.getBlockageUrls(
-				testBlockage,
-				this.externalRobomergeUrl,
-				"TEST",
-				"SOURCE_BRANCH",
-				"0",
-				false
+
+		if (this.slackMessages) {
+			let text = `${username}, please resolve the following test message:\n(source CL: 0, conflict CL: 1, shelf CL: 2)`
+
+			// This doesn't need to be a full Blockage -- just enough to generate required info for message
+			const testBlockage: Blockage = {
+				action: {
+					branch: {
+						name: "TARGET_BRANCH",
+					} as Branch,
+				} as MergeAction,
+				failure: {
+					kind: "Merge conflict",
+					description: ""
+				}
+			} as Blockage
+
+			const messageOpts = this.makeSlackDirectMessage(text, 0, 1, "TARGETBRANCH",
+				NodeBot.getBlockageUrls(
+					testBlockage,
+					this.externalRobomergeUrl,
+					"TEST",
+					"SOURCE_BRANCH",
+					"0",
+					false
+				)
 			)
-		)
 
-		this.slackMessages!.postDM([REDACTED_EMAIL]`, 0, "TARGETBRANCH", messageOpts, false)
+			this.slackMessages.postDM([REDACTED_EMAIL]`, 0, "TARGETBRANCH", messageOpts, false)
+		}
+		else {
+			this.botNotificationsLogger.error(`Slack not enabled. Unable to send test message to ${username}`)
+		}
 	}
 
 	sendGenericNonConflictMessage(message: string) {
@@ -688,7 +739,7 @@ export class BotNotifications implements BotEventHandler {
 
 		// Append footer
 		attachCollection.push({
-			pretext: "You can get help via the Slack channel <#C9321FLTU> (if you don't have access to 'robomerge-help', please contact the IT helpdesk)",
+			pretext: "You can get help via the Slack channel <#C044ZUNF9MJ> (if you don't have access to '#dt-robomerge-help-ext', please contact the IT helpdesk)",
 			mrkdwn_in: ["pretext"]
 		})
 
@@ -854,9 +905,19 @@ export class BotNotifications implements BotEventHandler {
 
 export function bindBotNotifications(events: BotEvents, slackChannelOverrides: [Branch, Branch, string, boolean][], persistence: Context, blockageUrlGenerator: NodeOpUrlGenerator, 
 	externalUrl: string, logger: ContextualLogger) {
-	events.registerHandler(new BotNotifications(events.botname, events.botConfig.slackChannel, persistence, externalUrl, blockageUrlGenerator, logger, slackChannelOverrides))
-}
 
+	let slackMessages
+
+	const botToken = args.devMode && SLACK_DEV_DUMMY_TOKEN || SLACK_TOKENS.bot
+	if (botToken && events.botConfig.slackChannel) {
+		logger.info('Enabling Slack messages for ' +  events.botname)
+		slackMessages = new SlackMessages(new Slack({id: events.botConfig.slackChannel, botToken}, args.slackDomain), persistence, logger)
+	}
+		
+	events.registerHandler(new BotNotifications(events.botname, events.botConfig.slackChannel, externalUrl, blockageUrlGenerator, logger, slackMessages, slackChannelOverrides))
+
+	return slackMessages
+}
 
 export function runTests(parentLogger: ContextualLogger) {
 	const unitTestLogger = parentLogger.createChild('Notifications')
@@ -871,7 +932,7 @@ export function runTests(parentLogger: ContextualLogger) {
 		kind: 'Unit Test error',
 
 		time: new Date,
-		nagged: false,
+		nagCount: 0,
 		ugsIssue: -1,
 
 		resolution: 'pickled' as Resolution
diff --git a/Engine/Extras/RoboMerge/v3/src/robo/slack.ts b/Engine/Extras/RoboMerge/v3/src/robo/slack.ts
--- a/Engine/Extras/RoboMerge/v3/src/robo/slack.ts
+++ b/Engine/Extras/RoboMerge/v3/src/robo/slack.ts
@@ -69,6 +69,10 @@ export class Slack {
 	constructor(private channel: SlackChannel, private domain: string) {
 	}
 
+	async addUserToChannel(user: string, channel: string) {
+		return this.post('conversations.invite', {channel: channel, users:user,}, true)
+	}
+
 	async postMessage(message: SlackMessage) {
 		return (await this.post('chat.postMessage', this.makeArgs(message))).ts
 	}
@@ -109,7 +113,7 @@ export class Slack {
 		return (await this.post('conversations.open', {token: this.channel.botToken, users})).channel.id
 	}
 
-	/*private*/ async post(command: string, args: any) {
+	/*private*/ async post(command: string, args: any, canFail? : boolean) {
 		const resultJson = await request.post({
 			url: this.domain + '/api/' + command,
 			body: JSON.stringify(args),
@@ -118,7 +122,7 @@ export class Slack {
 		})
 		try {
 			const result = JSON.parse(resultJson)
-			if (result.ok) {
+			if (result.ok || canFail) {
 				return result
 			}
 		}
