Store off the log before functional test cleans up

diff --git a/Engine/Extras/RoboMerge/v3/run-tests.sh b/Engine/Extras/RoboMerge/v3/run-tests.sh
--- a/Engine/Extras/RoboMerge/v3/run-tests.sh
+++ b/Engine/Extras/RoboMerge/v3/run-tests.sh
@@ -99,5 +99,7 @@ else
 
 fi
 
+docker logs robomerge_functtest > robo.log
+
 cleanup
 
