Make Slack get consistent with post in that a call can specify that it canFail
Fix up dummy slack server to be able to deal with add user to channel flow

diff --git a/Engine/Extras/RoboMerge/v3/src/common/dummyslackserver.ts b/Engine/Extras/RoboMerge/v3/src/common/dummyslackserver.ts
--- a/Engine/Extras/RoboMerge/v3/src/common/dummyslackserver.ts
+++ b/Engine/Extras/RoboMerge/v3/src/common/dummyslackserver.ts
@@ -16,6 +16,10 @@ export class DummySlackApp implements AppInterface {
 	post(command: string) {
 		// @todo (maybe) construct app with logger
 
+		if (command.startsWith("conversations.invite")) {
+			return { ok: true }
+		}
+
 		const data = this.req.reqData
 		if (!data) {
 			throw new Error('Nothing to post!')
@@ -29,7 +33,7 @@ export class DummySlackApp implements AppInterface {
 			}
 			catch (exc) {
 				console.log('Non-JSON sent to dummy-slack: ', data)
-				return
+				return {}
 			}
 
 			setDefault(
@@ -44,10 +48,23 @@ export class DummySlackApp implements AppInterface {
 		// look for color for block or resolve?
 
 		// careful, two changes for resolutions, luckily first is what we want
+
+		return {}
 	}
 
 	@Handler('GET', '/posted/*')
 	channelMessages(channel: string) {
 		return [...(posted.get(channel) || [])]
 	}
+
+	@Handler('GET', '/api/*')
+	get(command: string) {
+		console.log('dummy-slack GET: ', command)
+		if (command.startsWith("users.lookupByEmail"))
+		{
+			return { ok: true, user: { id: ""} }
+		}
+		return { ok: false }
+	}
+
 }
diff --git a/Engine/Extras/RoboMerge/v3/src/robo/slack.ts b/Engine/Extras/RoboMerge/v3/src/robo/slack.ts
--- a/Engine/Extras/RoboMerge/v3/src/robo/slack.ts
+++ b/Engine/Extras/RoboMerge/v3/src/robo/slack.ts
@@ -132,7 +132,7 @@ export class Slack {
 		throw new Error(`${command} generated:\n\t${resultJson}`)
 	}
 
-	/*private*/ async get(command: string, args: any) {
+	/*private*/ async get(command: string, args: any, canFail? : boolean) {
 
 		// erg: why am I always passing a channel?
 		if (this.channel.id && !args.channel) {
@@ -150,11 +150,16 @@ export class Slack {
 			headers: {Authorization: 'Bearer ' + this.channel.botToken}
 		})
 
-		const result = JSON.parse(rawResult)
-		if (!result.ok) {
-			throw new Error(`url: '${url}' error: '${rawResult}'`)
+		try {
+			const result = JSON.parse(rawResult)
+			if (result.ok || canFail) {
+				return result
+			}
+		}
+		catch {
 		}
-		return result
+		
+		throw new Error(`url: '${url}' error: '${rawResult}'`)
 	}
 
 	async* getPages(command: string, limit?: number, inArgs?: any) {
