Protect against failure to look up user

diff --git a/Engine/Extras/RoboMerge/v3/src/robo/nodebot.ts b/Engine/Extras/RoboMerge/v3/src/robo/nodebot.ts
--- a/Engine/Extras/RoboMerge/v3/src/robo/nodebot.ts
+++ b/Engine/Extras/RoboMerge/v3/src/robo/nodebot.ts
@@ -2024,7 +2024,13 @@ export class NodeBot extends PerforceStatefulBot implements NodeBotInterface {
 							const emailAddress = await this.findEmail(triager)
 							if (emailAddress)
 							{
-								triager = `<@${await this.slackMessages.getSlackUser(emailAddress)}>`
+								const user = await this.slackMessages.getSlackUser(emailAddress)
+								if (user) {
+									triager = `<@${user}>`
+								}
+								else {
+									this.nodeBotLogger.error(`Unable to look up triager from ${triager} (${emailAddress})`)
+								}
 							}
 						}
 
diff --git a/Engine/Extras/RoboMerge/v3/src/robo/notifications.ts b/Engine/Extras/RoboMerge/v3/src/robo/notifications.ts
--- a/Engine/Extras/RoboMerge/v3/src/robo/notifications.ts
+++ b/Engine/Extras/RoboMerge/v3/src/robo/notifications.ts
@@ -357,9 +357,14 @@ export class SlackMessages {
 	async addUserToChannel(emailAddress: string, channel: string)
 	{
 		const user = await this.getSlackUser(emailAddress)
-		const result = await this.slack.addUserToChannel(user, channel)
-		if (!result.ok && result.error !== "already_in_channel") {
-			this.smLogger.error(`Error inviting ${emailAddress} (${user}) to channel ${channel}: ${result.error}`)
+		if (user) {
+			const result = await this.slack.addUserToChannel(user, channel)
+			if (!result.ok && result.error !== "already_in_channel") {
+				this.smLogger.error(`Error inviting ${emailAddress} (${user}) to channel ${channel}: ${result.error}`)
+			}
+		}
+		else {
+			this.smLogger.error(`Unable to add ${emailAddress} to channel ${channel}`)
 		}
 	}
 
diff --git a/Engine/Extras/RoboMerge/v3/src/robo/slack.ts b/Engine/Extras/RoboMerge/v3/src/robo/slack.ts
--- a/Engine/Extras/RoboMerge/v3/src/robo/slack.ts
+++ b/Engine/Extras/RoboMerge/v3/src/robo/slack.ts
@@ -102,8 +102,9 @@ export class Slack {
 		return this.get('groups.history', args)
 	}
 
-	async lookupUserIdByEmail(email: string) : Promise<string> {
-		return (await this.get('users.lookupByEmail', {email})).user.id
+	async lookupUserIdByEmail(email: string) {
+		const userLookupResult = await this.get('users.lookupByEmail', {email}, true)
+		return userLookupResult.ok ? userLookupResult.user.id : null
 	}
 
 	async openDMConversation(users: string | string[]) : Promise<string> {
