Fix node resolver not being used when there are edge properties
Reconsider errors to direct messages
Correct slack style on existing error messages
Display bot alias in the graph

diff --git a/Engine/Extras/RoboMerge/v3/public/js/flow.js b/Engine/Extras/RoboMerge/v3/public/js/flow.js
--- a/Engine/Extras/RoboMerge/v3/public/js/flow.js
+++ b/Engine/Extras/RoboMerge/v3/public/js/flow.js
@@ -66,7 +66,8 @@ class Graph {
     singleBot(botName) {
         this.branchList = this.allBranches.filter(b => b.bot === botName);
         this.addNodesAndEdges(false);
-        return this.makeGraph(botName, botName + 'integration paths');
+        const title = this.options.aliases.length > 0 ? botName + ' (' + this.options.aliases.join(", ") + ')': botName
+        return this.makeGraph(title, botName + ' integration paths');
     }
     allBots() {
         this.branchList = this.options.botsToShow
diff --git a/Engine/Extras/RoboMerge/v3/public/js/main.js b/Engine/Extras/RoboMerge/v3/public/js/main.js
--- a/Engine/Extras/RoboMerge/v3/public/js/main.js
+++ b/Engine/Extras/RoboMerge/v3/public/js/main.js
@@ -121,7 +121,7 @@ function renderBranchList(data) {
 		let graphObj = branchGraphs.get(botName);
 		const branchSpecCl = botState.lastBranchspecCl;
 		if (!graphObj || graphObj.cl !== branchSpecCl) {
-			graphObj = {elem: showFlowGraph(branches, {singleBotName: botName}), cl: branchSpecCl};
+			graphObj = {elem: showFlowGraph(branches, {singleBotName: botName, aliases: botState.aliases}), cl: branchSpecCl};
 			branchGraphs.set(botName, graphObj);
 		}
 	
diff --git a/Engine/Extras/RoboMerge/v3/src/robo/graphbot.ts b/Engine/Extras/RoboMerge/v3/src/robo/graphbot.ts
--- a/Engine/Extras/RoboMerge/v3/src/robo/graphbot.ts
+++ b/Engine/Extras/RoboMerge/v3/src/robo/graphbot.ts
@@ -429,7 +429,8 @@ export class GraphBot implements GraphInterface, BotEventHandler {
 
 	applyStatus(out: Status) {
 		const status: GraphBotState = {
-			isRunningBots: this._runningBots
+			isRunningBots: this._runningBots,
+			aliases: this.branchGraph.config.aliases
 		}
 		if (this.autoUpdater) {
 			status.lastBranchspecCl = this.autoUpdater.lastCl
diff --git a/Engine/Extras/RoboMerge/v3/src/robo/nodebot.ts b/Engine/Extras/RoboMerge/v3/src/robo/nodebot.ts
--- a/Engine/Extras/RoboMerge/v3/src/robo/nodebot.ts
+++ b/Engine/Extras/RoboMerge/v3/src/robo/nodebot.ts
@@ -17,7 +17,7 @@ import { SlackMessages } from './notifications';
 import { PerforceStatefulBot } from './perforce-stateful-bot';
 import { BlockageNodeOpUrls, OperationUrlHelper } from './roboserver';
 import { Context } from './settings';
-import { SlackMessage } from './slack';
+import { SlackMessage, SlackMessageStyles } from './slack';
 import { PauseState } from './state-interfaces';
 import { newTickJournal, TickJournal } from './tick-journal';
 import { computeTargets, parseDescriptionLines, processOtherBotTargets, getIntegrationOwner, getNodeBotFullName, getNodeBotFullNameForLogging } from './targets';
@@ -363,23 +363,51 @@ export class NodeBot extends PerforceStatefulBot implements NodeBotInterface {
 			specifiedTargetBranch = this._getBranch(fromQueue.targetBranchName)
 
 			if (!specifiedTargetBranch) {
-				// We somehow have a specified branch from a reconsider request, but it doesn't exist in the branch map. Report the error
-				this.sendErrorEmail(
-					new Recipients(fromQueue.who),
-					`Invalid Reconsider Request -- Branch "${fromQueue.targetBranchName}" does not exist`,
-					`Manually queued change ${fromQueue.cl} on ${this.fullName}, requested by ${fromQueue.who}, ` + 
-						`specifies non-existant branch ${fromQueue.targetBranchName}. Aborting reconsider.`
-				)
+				if (this.slackMessages) {
+					const dm: SlackMessage = {
+						text: `Manually queued change ${fromQueue.cl} on ${this.fullName}, requested by ${fromQueue.who}, ` + 
+								`specifies non-existant branch ${fromQueue.targetBranchName}. Aborting reconsider.`,
+						style: SlackMessageStyles.DANGER,
+						channel: '',
+						mrkdwn: true
+					}
+					
+					let emailAddress = await this.findEmail(fromQueue.who)
+					this.slackMessages.postDM(emailAddress, fromQueue.cl, this.branch, dm)
+				}
+				else {
+					// We somehow have a specified branch from a reconsider request, but it doesn't exist in the branch map. Report the error
+					this.sendErrorEmail(
+						new Recipients(fromQueue.who),
+						`Invalid Reconsider Request -- Branch "${fromQueue.targetBranchName}" does not exist`,
+						`Manually queued change ${fromQueue.cl} on ${this.fullName}, requested by ${fromQueue.who}, ` + 
+							`specifies non-existant branch ${fromQueue.targetBranchName}. Aborting reconsider.`
+					)
+				}
 				return
 			}
 			else if (specifiedTargetBranch === this.branch) {
-				// We somehow have targeted our own branch with a merge. Report the issue
-				this.sendErrorEmail(
-					new Recipients(fromQueue.who),
-					`Invalid Reconsider Request -- Trying to merge "${fromQueue.targetBranchName}" into itself`,
-					`Manually queued change ${fromQueue.cl} on ${this.fullName}, requested by ${fromQueue.who}, ` +
-						`specifies integration from branch ${fromQueue.targetBranchName} into itself. Aborting reconsider.`
-				)
+				if (this.slackMessages) {
+					const dm: SlackMessage = {
+						text: `Manually queued change ${fromQueue.cl} on ${this.fullName}, requested by ${fromQueue.who}, ` +
+								`specifies integration from branch ${fromQueue.targetBranchName} into itself. Aborting reconsider.`,
+						style: SlackMessageStyles.DANGER,
+						channel: '',
+						mrkdwn: true
+					}
+					
+					let emailAddress = await this.findEmail(fromQueue.who)
+					this.slackMessages.postDM(emailAddress, fromQueue.cl, this.branch, dm)
+				}
+				else {
+					// We somehow have targeted our own branch with a merge. Report the issue
+					this.sendErrorEmail(
+						new Recipients(fromQueue.who),
+						`Invalid Reconsider Request -- Trying to merge "${fromQueue.targetBranchName}" into itself`,
+						`Manually queued change ${fromQueue.cl} on ${this.fullName}, requested by ${fromQueue.who}, ` +
+							`specifies integration from branch ${fromQueue.targetBranchName} into itself. Aborting reconsider.`
+					)
+				}
 				// Persist the remaining queued changes and carry on
 				return
 			}
@@ -393,6 +421,24 @@ export class NodeBot extends PerforceStatefulBot implements NodeBotInterface {
 		const changeResult = await this._getChange(fromQueue.cl)
 
 		if (typeof changeResult === 'string') {
+
+			if (this.slackMessages) {
+
+				let reconsiderErrorMessage = `While processing manually queued change ${fromQueue.cl} on ${this.fullName}, `
+				reconsiderErrorMessage += specifiedTargetBranch ? `targeted specifically at ${specifiedTargetBranch.name}` : ''
+				reconsiderErrorMessage += `\n\nError while querying P4 for change ${fromQueue.cl}: ${changeResult}`
+
+				const dm: SlackMessage = {
+					text: reconsiderErrorMessage,
+					style: SlackMessageStyles.DANGER,
+					channel: '',
+					mrkdwn: true
+				}
+				
+				let emailAddress = await this.findEmail(fromQueue.who)
+				this.slackMessages.postDM(emailAddress, fromQueue.cl, this.branch, dm)
+			}
+
 			// change is invalid
 
 			// @todo distinguish between Perforce errors and change not existing in branch
@@ -1861,6 +1907,7 @@ export class NodeBot extends PerforceStatefulBot implements NodeBotInterface {
 
 				let dm: SlackMessage = {
 					text: message,
+					style: SlackMessageStyles.DANGER,
 					channel: '',
 					mrkdwn: true
 				}
@@ -1915,6 +1962,7 @@ export class NodeBot extends PerforceStatefulBot implements NodeBotInterface {
 
 				let dm: SlackMessage = {
 					text: message,
+					style: SlackMessageStyles.WARNING,
 					channel: '',
 					mrkdwn: true
 				}
@@ -2036,6 +2084,7 @@ export class NodeBot extends PerforceStatefulBot implements NodeBotInterface {
 
 						const nagMessage: SlackMessage = {
 							text: `${triager} RoboMerge blocked for more than ${timeDesc}`,
+							style: SlackMessageStyles.DANGER,
 							channel: this.branchGraph.config.slackChannel,
 							mrkdwn: true
 						}
diff --git a/Engine/Extras/RoboMerge/v3/src/robo/notifications.ts b/Engine/Extras/RoboMerge/v3/src/robo/notifications.ts
--- a/Engine/Extras/RoboMerge/v3/src/robo/notifications.ts
+++ b/Engine/Extras/RoboMerge/v3/src/robo/notifications.ts
@@ -514,6 +514,7 @@ export class BotNotifications implements BotEventHandler {
 
 			let syntaxErrorMessage: SlackMessage = {
 				text: blockage.failure.description,
+				style: SlackMessageStyles.DANGER,
 				channel: message.channel,
 				mrkdwn: false
 			}
diff --git a/Engine/Extras/RoboMerge/v3/src/robo/status-types.ts b/Engine/Extras/RoboMerge/v3/src/robo/status-types.ts
--- a/Engine/Extras/RoboMerge/v3/src/robo/status-types.ts
+++ b/Engine/Extras/RoboMerge/v3/src/robo/status-types.ts
@@ -179,6 +179,7 @@ export interface GraphBotState {
 	isRunningBots: boolean
 	lastBranchspecCl?: number
 	lastError?: GraphBotError
+	aliases?: string[]
 }
 
 export type User = {userName: string, displayName: string, privileges?: string[]}
diff --git a/Engine/Extras/RoboMerge/v3/src/robo/targets.ts b/Engine/Extras/RoboMerge/v3/src/robo/targets.ts
--- a/Engine/Extras/RoboMerge/v3/src/robo/targets.ts
+++ b/Engine/Extras/RoboMerge/v3/src/robo/targets.ts
@@ -48,7 +48,7 @@ export function getIntegrationOwner(arg0: Branch | PendingChange, overriddenOwne
 	if (pending) {
 		const edgeprops = pending.change.branch.edgeProperties.get(branch!.upperName)
 		if (edgeprops) {
-			resolver = edgeprops.resolver || null
+			resolver = edgeprops.resolver || resolver || null
 		}	
 	}
 
diff --git a/Engine/Extras/RoboMerge/v3/web/src/flow.ts b/Engine/Extras/RoboMerge/v3/web/src/flow.ts
--- a/Engine/Extras/RoboMerge/v3/web/src/flow.ts
+++ b/Engine/Extras/RoboMerge/v3/web/src/flow.ts
@@ -107,7 +107,8 @@ class Graph {
 	singleBot(botName: string) {
 		this.branchList = this.allBranches.filter(b => b.bot === botName)
 		this.addNodesAndEdges(false)
-		return this.makeGraph(botName, botName + 'integration paths')
+        // TODO: Include aliases in graph title
+		return this.makeGraph(botName, botName + ' integration paths')
 	}
 
 	allBots() {
