Clear internal ugs url

diff --git a/Engine/Extras/RoboMerge/v3/src/robo/ugs.ts b/Engine/Extras/RoboMerge/v3/src/robo/ugs.ts
--- a/Engine/Extras/RoboMerge/v3/src/robo/ugs.ts
+++ b/Engine/Extras/RoboMerge/v3/src/robo/ugs.ts
@@ -5,7 +5,7 @@ import { Badge } from '../common/badge';
 import { ContextualLogger } from '../common/logger';
 
 // need to check whether build health reporting is working
-const UGS_URL_ROOT = 'http://ugsapi.epicgames.net'
+const UGS_URL_ROOT = '<ugs url root>'
 
 interface IntegrationIssue {
 	Project: string
