Fix nag notifications being able to be disabled

diff --git a/Engine/Extras/RoboMerge/v3/src/robo/nodebot.ts b/Engine/Extras/RoboMerge/v3/src/robo/nodebot.ts
--- a/Engine/Extras/RoboMerge/v3/src/robo/nodebot.ts
+++ b/Engine/Extras/RoboMerge/v3/src/robo/nodebot.ts
@@ -6,9 +6,9 @@ import { ContextualLogger } from '../common/logger';
 import { Mailer, MailParams, Recipients } from '../common/mailer';
 import { Change, ConflictedResolveNFile, PerforceContext, RoboWorkspace, coercePerforceWorkspace } from '../common/perforce';
 import { IPCControls, NodeBotInterface, QueuedChange, ReconsiderArgs } from './bot-interfaces';
-import { ApprovalOptions, BotConfig, EdgeOptions } from './branchdefs'
+import { ApprovalOptions, BotConfig, EdgeOptions, NodeOptions } from './branchdefs'
 import { BlockagePauseInfo, BranchStatus } from './status-types';
-import { AlreadyIntegrated, Blockage, Branch, BranchArg, BranchGraphInterface, ChangeInfo, EditableBranch, EndIntegratingToGateEvent, Failure } from './branch-interfaces';
+import { AlreadyIntegrated, Blockage, Branch, BranchArg, BranchGraphInterface, ChangeInfo, EndIntegratingToGateEvent, Failure } from './branch-interfaces';
 import { MergeAction, OperationResult, PendingChange, resolveBranchArg, StompedRevision, StompVerification, StompVerificationFile }  from './branch-interfaces';
 import { Conflicts } from './conflicts';
 import { EdgeBot, EdgeMergeResults } from './edgebot';
@@ -248,7 +248,7 @@ export class NodeBot extends PerforceStatefulBot implements NodeBotInterface {
 
 		// Pre-tick
 
-		if (this.edgesQueuedToUnblock) {
+		if (this.edgesQueuedToUnblock.size > 0) {
 			for (const edgeName of this.edgesQueuedToUnblock) {
 				const edge = this.getImmediateEdge(edgeName)
 				if (edge) {
@@ -2012,24 +2012,24 @@ export class NodeBot extends PerforceStatefulBot implements NodeBotInterface {
 						
 				const configKey = property as keyof BotConfig
 				const edgeKey   = property as keyof EdgeOptions
-				const branchKey = property as keyof EditableBranch
+				const nodeKey   = property as keyof NodeOptions
 
 				let prop: any | null = this.branchGraph.config[configKey]
 				if (conflict.targetBranchName) {
 					const targetBranch = this.branchGraph.getBranch(conflict.targetBranchName)
-					if (targetBranch)
+					if (targetBranch && targetBranch.config[nodeKey] !== undefined)
 					{
-						prop = targetBranch[branchKey] || prop
+						prop = targetBranch.config[nodeKey]
 					}
 
 					const edgeProperties = this.branch.edgeProperties.get(conflict.targetBranchName)
-					if (edgeProperties && edgeProperties[edgeKey]) {
-						prop = edgeProperties[edgeKey] || prop
+					if (edgeProperties && edgeProperties[edgeKey] !== undefined) {
+						prop = edgeProperties[edgeKey]
 					}
 				}
-				else
+				else if (this.branch.config[nodeKey] !== undefined)
 				{
-					prop = this.branch[branchKey] || prop
+					prop = this.branch.config[nodeKey]
 				}
 				return prop
 			}
