Check for resolved conflicts at a node level instead of an edge level to clear out stale conflicts due to syntax errors or on edges that no longer exist.
Also strip on load very old conflicts that don't have all the expected persisted data

diff --git a/Engine/Extras/RoboMerge/v3/src/robo/conflicts.ts b/Engine/Extras/RoboMerge/v3/src/robo/conflicts.ts
--- a/Engine/Extras/RoboMerge/v3/src/robo/conflicts.ts
+++ b/Engine/Extras/RoboMerge/v3/src/robo/conflicts.ts
@@ -103,12 +103,14 @@ export class Conflicts {
 		const conflictsToLoad = persistence.get('conflicts') as PersistentConflict[]
 		if (conflictsToLoad) {
 			for (const conflict of conflictsToLoad) {
-				conflict.time = new Date(conflict.time)
-				if (conflict.lastNagTime) {
-					conflict.lastNagTime = new Date(conflict.lastNagTime)
+				if (conflict.blockedBranchName) { // Discard old conflicts in persistent data that don't have expected fields
+					conflict.time = new Date(conflict.time)
+					if (conflict.lastNagTime) {
+						conflict.lastNagTime = new Date(conflict.lastNagTime)
+					}
+					conflict.nagCount = conflict.nagCount || 0 // back compat for nag changes
+					this.conflicts.push(conflict)
 				}
-				conflict.nagCount = conflict.nagCount || 0 // back compat for nag changes
-				this.conflicts.push(conflict)
 			}
 		}
 
@@ -228,22 +230,30 @@ export class Conflicts {
 	 * This happens before a tick, so that all other bots will have had a chance
 	 * to update, so that resolving changelists will have been seen
 	 */
-	checkForResolvedConflicts(edgeBot: EdgeBot) {
+	checkForResolvedConflicts(nodeCL: Number, edges: Map<string, EdgeBot>) {
 		if (this.conflicts.length === 0) {
 			return
 		}
 
-		const targetBranchStr = resolveBranchArg(edgeBot.targetBranch, true)
-
 		// Check all the conflicts from our branch for resolved blockages
 		const remainingConflicts: PersistentConflict[] = []
 		const resolvedConflicts: PersistentConflict[] = []
 
 		for (const conflict of this.conflicts) {
-			// If our bot CL is below the conflict CL, we can assume we haven't resolved this conflict
-			if (!conflict.targetBranchName || conflict.targetBranchName.toUpperCase() !== targetBranchStr || edgeBot.lastCl < conflict.cl) {
-				remainingConflicts.push(conflict)
-				continue
+
+			if (!conflict.targetBranchName ) {
+				if (nodeCL < conflict.cl) {
+					remainingConflicts.push(conflict)
+					continue
+				}
+			}
+			else {
+				// If our bot CL is below the conflict CL, we can assume we haven't resolved this conflict
+				const edgeBot = edges.get(conflict.targetBranchName)		
+				if (edgeBot && edgeBot.lastCl < conflict.cl) {
+					remainingConflicts.push(conflict)
+					continue
+				}
 			}
 
 			// If the conflict didn't fill out the resolution, try to determine one now
@@ -265,7 +275,7 @@ export class Conflicts {
 				setResolvedTimeAsNow(conflict)
 			}
 
-			this.conflictLogger.info(`Conflict for branch ${targetBranchStr} cl ${conflict.cl} seems to be resolved: ${conflict.resolution} by ${conflict.resolvingAuthor}`)
+			this.conflictLogger.info(`Conflict for branch ${conflict.targetBranchName} cl ${conflict.cl} seems to be resolved: ${conflict.resolution} by ${conflict.resolvingAuthor}`)
 		}
 
 
@@ -309,7 +319,7 @@ export class Conflicts {
 				conflict.resolution = Resolution.CANCELLED
 				setResolvedTimeAsNow(conflict)
 				this.setConflicts(this.conflicts.filter(el => el !== conflict))
-this.conflictLogger.info('Reporting unblocked (cancelled) ' + conflict.cl)
+				this.conflictLogger.info('Reporting unblocked (cancelled) ' + conflict.cl)
 				this.reportUnblocked(conflict)
 			}
 		}
diff --git a/Engine/Extras/RoboMerge/v3/src/robo/nodebot.ts b/Engine/Extras/RoboMerge/v3/src/robo/nodebot.ts
--- a/Engine/Extras/RoboMerge/v3/src/robo/nodebot.ts
+++ b/Engine/Extras/RoboMerge/v3/src/robo/nodebot.ts
@@ -267,8 +267,9 @@ export class NodeBot extends PerforceStatefulBot implements NodeBotInterface {
 		let minCl = this.lastCl
 		// Gather the list of available edges for tick -- edges may change mid-tick but we will only allow this set to perform integrations
 		let availableEdges: Map<string, EdgeBot> = new Map<string, EdgeBot>()
+
+		this.conflicts.checkForResolvedConflicts(this.lastCl, this.edges)
 		for (const edgeBot of this.edges.values()) {
-			this.conflicts.checkForResolvedConflicts(edgeBot)
 			if (edgeBot.isAvailable) {
 				availableEdges.set(edgeBot.targetBranch.upperName, edgeBot)
 				if (edgeBot.lastCl > 0 && (minCl < 0 || edgeBot.lastCl < minCl)) {
