Fix message when triager not specified

diff --git a/Engine/Extras/RoboMerge/v3/src/robo/nodebot.ts b/Engine/Extras/RoboMerge/v3/src/robo/nodebot.ts
--- a/Engine/Extras/RoboMerge/v3/src/robo/nodebot.ts
+++ b/Engine/Extras/RoboMerge/v3/src/robo/nodebot.ts
@@ -2084,7 +2084,7 @@ export class NodeBot extends PerforceStatefulBot implements NodeBotInterface {
 						}
 
 						const nagMessage: SlackMessage = {
-							text: `${triager} RoboMerge blocked for more than ${timeDesc}`,
+							text: `${triager ? triager + ' ' : ''}RoboMerge blocked for more than ${timeDesc}`,
 							style: SlackMessageStyles.DANGER,
 							channel: this.branchGraph.config.slackChannel,
 							mrkdwn: true
