Allow stomp in binary branch and delete cases
Text files will still fail due to the standard text file restrictions

Ping users with the locked files on exclusive check outs
Fix some exclusive checked out files not reporting (presumably those on edge servers?)

Fix channel IDs to be consistent when initial post and nagging when there are multi-step integrations that cause a blockage

Aliases now displayed on allbots page

Functional tests now use 2022.1 perforce server

diff --git a/Engine/Extras/RoboMerge/v3/functional_tests/framework/src/framework.ts b/Engine/Extras/RoboMerge/v3/functional_tests/framework/src/framework.ts
--- a/Engine/Extras/RoboMerge/v3/functional_tests/framework/src/framework.ts
+++ b/Engine/Extras/RoboMerge/v3/functional_tests/framework/src/framework.ts
@@ -1009,6 +1009,15 @@ export class P4Util {
 		return this.editFile(client, name, newContent)
 			.then(() => this.submit(client, commitMessage))
 	}
+
+	static deleteFileAndSubmit(client: P4Client, name: string, robomergeCommand?: string) {
+		let commitMessage = `Delete file '${name}'`
+		if (robomergeCommand) {
+			commitMessage += '\n#robomerge ' + robomergeCommand
+		}
+		return client.delete(name)
+			.then(() => this.submit(client, commitMessage))
+	}
 	
 	static escapeBranchName(stream: string) {
 		return stream.replace(/[^A-Za-z0-9]/g, '_')
diff --git a/Engine/Extras/RoboMerge/v3/functional_tests/framework/src/tests.ts b/Engine/Extras/RoboMerge/v3/functional_tests/framework/src/tests.ts
--- a/Engine/Extras/RoboMerge/v3/functional_tests/framework/src/tests.ts
+++ b/Engine/Extras/RoboMerge/v3/functional_tests/framework/src/tests.ts
@@ -28,8 +28,6 @@ import { MultipleConflicts } from './tests/multiple-conflicts'
 import { MultipleRoutesToSkip } from './tests/multiple-routes-to-skip'
 import { OverriddenCommand } from './tests/overridden-command'
 import { PostToAdditionalChannel } from './tests/post-to-additional-channel'
-import { RejectBranchResolveStomp } from './tests/reject-branch-resolve-stomp'
-import { RejectDeleteResolveStomp } from './tests/reject-delete-resolve-stomp'
 import { RejectTextConflictStomp } from './tests/reject-text-conflict-stomp'
 import { RequestShelf } from './tests/request-shelf'
 import { RequestShelfIndirectTarget } from './tests/request-shelf-indirect-target'
@@ -38,6 +36,8 @@ import { RespectStreamPath } from './tests/respect-stream-path'
 import { SyntaxErrorOnUnknownBranch } from './tests/syntax-error-on-unknown-branch'
 import { StompForwardingCommands } from './tests/stomp-forwarding-commands'
 import { StompWithAdd } from './tests/stomp-with-add'
+import { StompWithBranch } from './tests/stomp-with-branch'
+import { StompWithDelete } from './tests/stomp-with-delete'
 import { TestChain } from './tests/test-chain'
 import { TestEdgeGate } from './tests/test-edge-gate'
 import { TestFlags } from './tests/test-flags'
@@ -172,8 +172,8 @@ async function go() {
 
 		new MergeMainRevToRelease(p4),
 		new MultipleConflicts(p4),
-		new RejectBranchResolveStomp(p4),
-		new RejectDeleteResolveStomp(p4),
+		new StompWithBranch(p4),
+		new StompWithDelete(p4),
 		new RejectTextConflictStomp(p4), // 15
 
 		new RequestShelf(p4),
diff --git a/Engine/Extras/RoboMerge/v3/functional_tests/framework/src/tests/reject-branch-resolve-stomp.ts b/Engine/Extras/RoboMerge/v3/functional_tests/framework/src/tests/reject-branch-resolve-stomp.ts
deleted file mode 100644
diff --git a/Engine/Extras/RoboMerge/v3/functional_tests/framework/src/tests/reject-delete-resolve-stomp.ts b/Engine/Extras/RoboMerge/v3/functional_tests/framework/src/tests/reject-delete-resolve-stomp.ts
deleted file mode 100644
diff --git a/Engine/Extras/RoboMerge/v3/functional_tests/framework/src/tests/stomp-with-branch.ts b/Engine/Extras/RoboMerge/v3/functional_tests/framework/src/tests/stomp-with-branch.ts
new file mode 100644
--- /dev/null
+++ b/Engine/Extras/RoboMerge/v3/functional_tests/framework/src/tests/stomp-with-branch.ts
@@ -0,0 +1,28 @@
+// Copyright Epic Games, Inc. All Rights Reserved.
+import { P4Util } from '../framework'
+import { SimpleMainAndReleaseTestBase, streams } from '../SimpleMainAndReleaseTestBase'
+
+export class StompWithBranch extends SimpleMainAndReleaseTestBase {
+
+	async setup() {
+		// Set up depot
+
+		await this.p4.depot('stream', this.depotSpec())
+		await this.createStreamsAndWorkspaces(streams)
+		await P4Util.addFileAndSubmit(this.getClient('Main'), 'fake.uasset', 'Initial content', true)
+		await this.initialPopulate()
+		await P4Util.deleteFileAndSubmit(this.getClient('Main'), 'fake.uasset')
+	}
+
+	async run() {
+		await this.getClient('Release').sync()
+		await P4Util.editFileAndSubmit(this.getClient('Release'), 'fake.uasset', 'New content')
+		await this.waitForRobomergeIdle()
+		await this.ensureBlocked('Release', 'Main')
+		await this.verifyAndPerformStomp('Release', 'Main')
+	}
+
+	async verify() {
+		await this.ensureNotBlocked('Release', 'Main')
+	}
+}
diff --git a/Engine/Extras/RoboMerge/v3/functional_tests/framework/src/tests/stomp-with-delete.ts b/Engine/Extras/RoboMerge/v3/functional_tests/framework/src/tests/stomp-with-delete.ts
new file mode 100644
--- /dev/null
+++ b/Engine/Extras/RoboMerge/v3/functional_tests/framework/src/tests/stomp-with-delete.ts
@@ -0,0 +1,28 @@
+// Copyright Epic Games, Inc. All Rights Reserved.
+import { P4Util } from '../framework'
+import { SimpleMainAndReleaseTestBase, streams } from '../SimpleMainAndReleaseTestBase'
+
+export class StompWithDelete extends SimpleMainAndReleaseTestBase {
+
+	async setup() {
+		// Set up depot
+
+		await this.p4.depot('stream', this.depotSpec())
+		await this.createStreamsAndWorkspaces(streams)
+		await P4Util.addFileAndSubmit(this.getClient('Main'), 'fake.uasset', 'Initial content', true)
+		await this.initialPopulate()
+		await P4Util.editFileAndSubmit(this.getClient('Main'), 'fake.uasset', 'New content')
+	}
+
+	async run() {
+		await this.getClient('Release').sync()
+		await P4Util.deleteFileAndSubmit(this.getClient('Release'), 'fake.uasset')
+		await this.waitForRobomergeIdle()
+		await this.ensureBlocked('Release', 'Main')
+		await this.verifyAndPerformStomp('Release', 'Main')
+	}
+
+	async verify() {
+		await this.ensureNotBlocked('Release', 'Main')
+	}
+}
diff --git a/Engine/Extras/RoboMerge/v3/public/allbots.html b/Engine/Extras/RoboMerge/v3/public/allbots.html
--- a/Engine/Extras/RoboMerge/v3/public/allbots.html
+++ b/Engine/Extras/RoboMerge/v3/public/allbots.html
@@ -21,6 +21,7 @@
 			success: data => {
 				const bots = location.search.match(/bots=([^&]*)/);
 				$(() => $('#graph').append(showFlowGraph(data.branches, {
+					aliases: new Map([...data.botStates].map(([botName, botState]) => [botName, botState.aliases])),
 					botNames: bots ? bots[1].split(',') : [],
 					...parseOptions(location.search)
 				})));
diff --git a/Engine/Extras/RoboMerge/v3/public/js/flow.js b/Engine/Extras/RoboMerge/v3/public/js/flow.js
--- a/Engine/Extras/RoboMerge/v3/public/js/flow.js
+++ b/Engine/Extras/RoboMerge/v3/public/js/flow.js
@@ -79,7 +79,13 @@ class Graph {
     }
     addNodesAndEdges(grouped) {
         for (const branchStatus of this.branchList) {
-            this.addBranch(branchStatus, grouped ? branchStatus.bot : null);
+            const botTitle = grouped
+                                ? this.options.aliases.get(branchStatus.bot).length > 0
+                                    ? branchStatus.bot + ' (' + this.options.aliases.get(branchStatus.bot).join(", ") + ')'
+                                    : branchStatus.bot
+                                : null;
+
+            this.addBranch(branchStatus, grouped ? branchStatus.bot : null, botTitle);
         }
         for (const info of this.allInfos) {
             for (const sourceBranch of info.sources) {
@@ -120,8 +126,8 @@ class Graph {
         for (const [groupName, nodeInfo] of nodeGroups) {
             if (groupName !== 'nogroup') {
                 lines.push(`subgraph cluster_${groupName} {
-	label="${groupName}";
-`);
+                            label="${nodeInfo[0].info.title}";
+                        `);
             }
             for (const info of nodeInfo) {
                 if (this.options.hideDisconnected && !this.connectedNodes.has(info.info.id)) {
@@ -160,7 +166,7 @@ class Graph {
         lines.push('}');
         return lines;
     }
-    addBranch(branchStatus, group = null) {
+    addBranch(branchStatus, group = null, title = null) {
         const branch = branchStatus.def;
         if (this.aliases.has(decoratedAlias(branchStatus.bot, branch.upperName)))
             return;
@@ -169,6 +175,7 @@ class Graph {
             tooltip += ` (${branch.aliases.join(', ')})`;
         }
         const info = new Info(`_${branchStatus.bot}_${branch.upperName.replace(/[^\w]+/g, '_')}`, tooltip, branch.name, [branchStatus]);
+        info.title = title
         this.allInfos.push(info);
         if (branch.config.graphNodeColor) {
             info.graphNodeColor = branch.config.graphNodeColor;
diff --git a/Engine/Extras/RoboMerge/v3/public/js/stomp.js b/Engine/Extras/RoboMerge/v3/public/js/stomp.js
--- a/Engine/Extras/RoboMerge/v3/public/js/stomp.js
+++ b/Engine/Extras/RoboMerge/v3/public/js/stomp.js
@@ -263,12 +263,6 @@ function visualizeStompVerification(requestedBranchCl, stompJson) {
             )
         }
 
-        if (file.branchOrDeleteResolveRequired) {
-            $('#branchOrDeleteResolveRequiredList').append(
-                $('<li>').html(linkToSwarm)
-            )
-        }
-
         if (file.stompedRevisionsSkipped) {
             $('#stompedRevisionsSkippedDiv').removeClass('initiallyHidden')
             $('#stompedRevisionsSkippedList').append(
@@ -309,14 +303,6 @@ function visualizeStompVerification(requestedBranchCl, stompJson) {
         }
     }
 
-    if (stompJson.branchOrDeleteResolveRequired) {
-        $('#verificationBranchOrDeleteResolveRequired').removeClass('initiallyHidden')
-        $('#resultVisualization').append(
-            $('<h5 class="centered-text">').html(`Robomerge cannot proceed with the stomp request using ${linkToSwarm}.`)
-        )
-        return
-    }
-
     if (encounteredStompedRevisionsCalculationIssues) {
         $('#stompedRevisionsIssuesDiv').removeClass('initiallyHidden')
     }
diff --git a/Engine/Extras/RoboMerge/v3/public/stomp.html b/Engine/Extras/RoboMerge/v3/public/stomp.html
--- a/Engine/Extras/RoboMerge/v3/public/stomp.html
+++ b/Engine/Extras/RoboMerge/v3/public/stomp.html
@@ -70,21 +70,6 @@
             </div> 
         </div>
 
-        <!-- Display error to user if the remaining files are not all binary, resulting in a refusal to stomp -->
-        <div id="verificationBranchOrDeleteResolveRequired" class="card card-single border-danger centered-text initiallyHidden">
-            <div class="card-header bg-danger text-white">
-               <span><i class="fas fa-exclamation-circle"></i></span>
-                <span class="sr-only">Error:</span>
-                Merge requires branching or deleting binary files.
-            </div>
-            <div class="card-body centered-text">
-                <p>Robomerge found conflicts while integrating changelist <strong class="changelist"></strong> which require either branching new copies of files over previously deleted revisions and/or deleting revisions present in <strong class="targetName"></strong>:</p>
-                <ul id="branchOrDeleteResolveRequiredList"></ul>
-                <p>Robomerge does not support this operation at this time. (Sorry!) You'll need to resolve this by hand.</p> <!-- To get started, <a href="/op/createshelf?">Robomerge could create a shelf for you.</a> -->
-                <p>If you still require assistance, please <strong><a href="/contact">Contact Us</a></strong> for Robomerge help.</p>
-            </div> 
-        </div>
-
         <!-- Display warning to user if some files had their stomped revision calculations skipped (task streams) -->
         <div id="stompedRevisionsIssuesDiv" class="card card-single border-warning centered-text initiallyHidden">
                 <div class="card-header bg-warning text-white">
diff --git a/Engine/Extras/RoboMerge/v3/src/common/dummyslackserver.ts b/Engine/Extras/RoboMerge/v3/src/common/dummyslackserver.ts
--- a/Engine/Extras/RoboMerge/v3/src/common/dummyslackserver.ts
+++ b/Engine/Extras/RoboMerge/v3/src/common/dummyslackserver.ts
@@ -59,10 +59,9 @@ export class DummySlackApp implements AppInterface {
 
 	@Handler('GET', '/api/*')
 	get(command: string) {
-		console.log('dummy-slack GET: ', command)
 		if (command.startsWith("users.lookupByEmail"))
 		{
-			return { ok: true, user: { id: ""} }
+			return { ok: true, user: { id: "" } }
 		}
 		return { ok: false }
 	}
diff --git a/Engine/Extras/RoboMerge/v3/src/common/perforce.ts b/Engine/Extras/RoboMerge/v3/src/common/perforce.ts
--- a/Engine/Extras/RoboMerge/v3/src/common/perforce.ts
+++ b/Engine/Extras/RoboMerge/v3/src/common/perforce.ts
@@ -294,7 +294,6 @@ export interface ConflictedResolveNFile {
 	resolveType: string // e.g. 'content', 'branch'
 	resolveFlag: string // 'c' ?
 	contentResolveType?: string  // e.g. '3waytext', '2wayraw' -- not applicable for delete/branch merges
-	branchOrDeleteResolveRequired?: boolean
 }
 
 export class ResolveResult {
@@ -341,7 +340,7 @@ export class ResolveResult {
 			at: branch
 			ay: ignore
 		
-		Skip that but note it in 'branchOrDeleteResolveRequired'.
+		We will allow these to be stomped
 	*/
 	private parseConflictsFromDashNOutput() {
 		this.successfullyParsedRemainingConflicts = true
@@ -349,10 +348,8 @@ export class ResolveResult {
 			for (let ztagGroup of this.dashNOutput) {
 				if (ztagGroup[0] === "Branch resolve:" || ztagGroup[0] === "Delete resolve:") {
 					if (!this.remainingConflicts[this.remainingConflicts.length - 1]) {
-						throw new Error(`Encountered branch/delete resolve information, but could not find appicable conflict file: ${ztagGroup.toString()}`)
+						throw new Error(`Encountered branch/delete resolve information, but could not find applicable conflict file: ${ztagGroup.toString()}`)
 					}
-
-					this.remainingConflicts[this.remainingConflicts.length - 1].branchOrDeleteResolveRequired = true
 					continue
 				}
 
@@ -434,6 +431,8 @@ export function getPerforceUsername() {
 	return perforceUsername
 }
 
+let perforceMultiServerEnvironment: boolean
+
 /**
  * This method must succeed before PerforceContext can be used. Otherwise retrieving the Perforce username through
  * getPerforceUsername() will error.
@@ -454,8 +453,10 @@ export async function initializePerforce(logger: ContextualLogger) {
 
 	if (resp && resp.User) {
 		perforceUsername = resp.User;
+
+		const serversOutput = await PerforceContext._execP4Ztag(logger, null, ["servers"]);
+		perforceMultiServerEnvironment = serversOutput.length > 1
 	}
-	return resp;
 }
 
 /**
@@ -1061,7 +1062,7 @@ export class PerforceContext {
 	}
 
 	// run p4 'opened' command: lists files in changelist with details of edit state (e.g. if a copy, provides source path)
-	opened(roboWorkspace: RoboWorkspace | null, arg: number | string) {
+	opened(roboWorkspace: RoboWorkspace | null, arg: number | string, exclusive?: boolean) {
 		const workspace = roboWorkspace && coercePerforceWorkspace(roboWorkspace);
 		const args = ['opened']
 		if (typeof arg === 'number') {
@@ -1070,7 +1071,7 @@ export class PerforceContext {
 		}
 		else {
 			// see which workspace has a file checked out/added
-			args.push('-a', arg)
+			args.push(exclusive && perforceMultiServerEnvironment ? '-x' : '-a', arg)
 		}
 		return this._execP4Ztag(workspace, args) as Promise<OpenedFileRecord[]>
 	}
diff --git a/Engine/Extras/RoboMerge/v3/src/robo/branch-interfaces.ts b/Engine/Extras/RoboMerge/v3/src/robo/branch-interfaces.ts
--- a/Engine/Extras/RoboMerge/v3/src/robo/branch-interfaces.ts
+++ b/Engine/Extras/RoboMerge/v3/src/robo/branch-interfaces.ts
@@ -45,8 +45,6 @@ export interface StompVerificationFile {
 	stompedRevisionsCalculationIssues: boolean
 	
 	targetFileName : string
-
-	branchOrDeleteResolveRequired: boolean // Currently we don't have a good way to handle branch / delete merges. Fail those requests for now.
 }
 
 export interface StompVerification extends OperationResult {
@@ -55,7 +53,6 @@ export interface StompVerification extends OperationResult {
 	svFiles?: StompVerificationFile[] // Array of remaining files
 
 	// Convinence booleans to alert user to problems in verification result
-	branchOrDeleteResolveRequired?: boolean // Currently we don't have a good way to handle branch / delete merges. Fail those requests for now.
 	remainingAllBinary?: boolean // Check to see if unresolved non-binary files are remaining -- we shouldn't stomp those!
 }
 
@@ -163,6 +160,7 @@ export interface Failure {
 	kind: FailureKind		// short description of integration error or conflict
 	description: string		// detailed description (can be very long - don't want to store this)
 	summary?: string
+	additionalInfo?: any
 }
 
 export interface Blockage {
diff --git a/Engine/Extras/RoboMerge/v3/src/robo/edgebot.ts b/Engine/Extras/RoboMerge/v3/src/robo/edgebot.ts
--- a/Engine/Extras/RoboMerge/v3/src/robo/edgebot.ts
+++ b/Engine/Extras/RoboMerge/v3/src/robo/edgebot.ts
@@ -280,8 +280,7 @@ class EdgeBotImpl extends PerforceStatefulBot {
 		for (const err of errors.slice(0, MAX_INTEGRATION_ERRORS_TO_ANALYZE)) {
 			const match = err.match(EXCLUSIVE_CHECKOUT_REGEX)
 			if (match) {
-				console.log(match)
-				openedRequests.push([match, this.p4.opened(null, match[1] + match[2])])
+				openedRequests.push([match, this.p4.opened(null, match[1] + match[2], true)])
 			}
 		}
 
@@ -557,7 +556,8 @@ class EdgeBotImpl extends PerforceStatefulBot {
 			if (errors.length > MAX_INTEGRATION_ERRORS_TO_ANALYZE) {
 				exclCheckoutMessages.push(`... and up to ${errors.length - MAX_INTEGRATION_ERRORS_TO_ANALYZE} more`)
 			}
-			failure = { kind: 'Exclusive check-out', description, summary: exclCheckoutMessages.join('\n') }
+			const exclCheckoutUsers = Array.from(new Set(exclusiveFiles.map(exc => `${exc.user.toLowerCase()}`))).map(user => ({user, userEmail: this.p4.getEmail(user)}))
+			failure = { kind: 'Exclusive check-out', description, summary: exclCheckoutMessages.join('\n'), additionalInfo: exclCheckoutUsers }
 		}
 		else {
 			failure  = { kind: 'Integration error', description }
diff --git a/Engine/Extras/RoboMerge/v3/src/robo/ipc.ts b/Engine/Extras/RoboMerge/v3/src/robo/ipc.ts
--- a/Engine/Extras/RoboMerge/v3/src/robo/ipc.ts
+++ b/Engine/Extras/RoboMerge/v3/src/robo/ipc.ts
@@ -463,7 +463,6 @@ export class IPC {
 						message: stompVerification.message,
 						nonBinaryFilesResolved: stompVerification.nonBinaryFilesResolved,
 						remainingAllBinary: stompVerification.remainingAllBinary,
-						branchOrDeleteResolveRequired: stompVerification.branchOrDeleteResolveRequired,
 						files: stompVerification.svFiles,
 						validRequest: stompVerification.validRequest
 					} )
diff --git a/Engine/Extras/RoboMerge/v3/src/robo/nodebot.ts b/Engine/Extras/RoboMerge/v3/src/robo/nodebot.ts
--- a/Engine/Extras/RoboMerge/v3/src/robo/nodebot.ts
+++ b/Engine/Extras/RoboMerge/v3/src/robo/nodebot.ts
@@ -764,7 +764,6 @@ export class NodeBot extends PerforceStatefulBot implements NodeBotInterface {
 			}
 			
 			// Process all the files in our pending changelist
-			let branchOrDeleteResolveRequired = false
 			let stompedRevisionsCalculationIssues = false
 			for (const file of describeResult.entries) {
 				// Check if this file is still remaining in the changelist with a conflict
@@ -774,7 +773,6 @@ export class NodeBot extends PerforceStatefulBot implements NodeBotInterface {
 				
 				let stompedRevisions : StompedRevision[] | null = []
 				let stompVerificationFile : StompVerificationFile = {
-					branchOrDeleteResolveRequired: false,
 					filetype: file.type,
 					stompedRevisionsSkipped: false,
 					stompedRevisionsCalculationIssues: false,
@@ -793,12 +791,8 @@ export class NodeBot extends PerforceStatefulBot implements NodeBotInterface {
 				} 
 
 				if (!resolved) {
-					if (remainingFile!.branchOrDeleteResolveRequired) {
-						stompVerificationFile.branchOrDeleteResolveRequired = true
-						branchOrDeleteResolveRequired = true
-					}
 					// Skip stomped revisions for task streams
-					else if (file.depotFile.startsWith('//Tasks/')) {
+					if (file.depotFile.startsWith('//Tasks/')) {
 						stompVerificationFile.stompedRevisionsSkipped = true
 						stompedRevisions = null
 					}
@@ -854,15 +848,10 @@ export class NodeBot extends PerforceStatefulBot implements NodeBotInterface {
 				message: "TBD",
 				nonBinaryFilesResolved,
 				remainingAllBinary,
-				branchOrDeleteResolveRequired,
 				svFiles
 			}
 
-			if (branchOrDeleteResolveRequired) {
-				returnResult.message= "Invalid Request: Currently do not support stomps requiring branch/delete resolution. Please merge manually."
-				returnResult.validRequest = false
-			} 
-			else if (!remainingAllBinary) {
+			if (!remainingAllBinary) {
 				returnResult.message= "Invalid Request: Not all remaining files are binary"
 				returnResult.validRequest = false
 			}
diff --git a/Engine/Extras/RoboMerge/v3/src/robo/notifications.ts b/Engine/Extras/RoboMerge/v3/src/robo/notifications.ts
--- a/Engine/Extras/RoboMerge/v3/src/robo/notifications.ts
+++ b/Engine/Extras/RoboMerge/v3/src/robo/notifications.ts
@@ -223,8 +223,8 @@ export class SlackMessages {
 		}
 	}
 
-	async postOrUpdate(sourceCl: number, branchArg: BranchArg, message: SlackMessage, persistMessage = true) {
-		const findResult = this.find(sourceCl, branchArg, message.channel)
+	async postOrUpdate(cl: number, branchArg: BranchArg, message: SlackMessage, persistMessage = true) {
+		const findResult = this.find(cl, branchArg, message.channel)
 
 		// If we find a message, simply update the contents
 		if (findResult.messageRecord) {
@@ -263,8 +263,8 @@ export class SlackMessages {
 		}
 	}
 
-	async postReply(sourceCl: number, branchArg: BranchArg, message: SlackMessage) {
-		const findResult = this.find(sourceCl, branchArg, message.channel)
+	async postReply(cl: number, branchArg: BranchArg, message: SlackMessage) {
+		const findResult = this.find(cl, branchArg, message.channel)
 
 		if (findResult.messageRecord) {
 			try {
@@ -276,15 +276,15 @@ export class SlackMessages {
 			}
 		}
 		else {
-			this.smLogger.error(`Failed to find message record for ${sourceCl}:${branchArg}:${message.channel}`)
+			this.smLogger.error(`Failed to find message record for ${cl}:${branchArg}:${message.channel}`)
 		}
 	}
 
-	async postDM(emailAddress: string|null, sourceCl: number, branchArg: BranchArg, dm: SlackMessage, persistMessage = true) {
+	async postDM(emailAddress: string|null, cl: number, branchArg: BranchArg, dm: SlackMessage, persistMessage = true) {
 		// The Slack API requires a user ID to open a direct message with users.
 		// The most consistent way to do this is getting their email address out of P4.
 		if (!emailAddress) {
-			console.error("Failed to get email address during notifications for CL " + sourceCl)
+			console.error("Failed to get email address during notifications for CL " + cl)
 			return
 		}
 
@@ -292,7 +292,7 @@ export class SlackMessages {
 		try {
 			userId = (await this.getSlackUser(emailAddress))
 		} catch (err) {
-			this.smLogger.printException(err, `Failed to get user ID for Slack DM, given email address "${emailAddress}" for CL ${sourceCl}`)
+			this.smLogger.printException(err, `Failed to get user ID for Slack DM, given email address "${emailAddress}" for CL ${cl}`)
 			return
 		}
 
@@ -302,20 +302,20 @@ export class SlackMessages {
 			channelId = (await this.slack.openDMConversation(userId))
 			dm.channel = channelId
 		} catch (err) {
-			this.smLogger.printException(err, `Failed to get Slack conversation ID for user ID "${userId}" given email address "${emailAddress}" for CL ${sourceCl}`)
+			this.smLogger.printException(err, `Failed to get Slack conversation ID for user ID "${userId}" given email address "${emailAddress}" for CL ${cl}`)
 			return
 		}
 
 		// Add the channel/conversation ID to the messageOpts and proceed normally.
-		this.smLogger.info(`Creating direct message for ${emailAddress} (key: ${generatePersistedSlackMessageKey(sourceCl, branchArg, channelId)})`)
-		this.postOrUpdate(sourceCl, branchArg, dm, persistMessage)
+		this.smLogger.info(`Creating direct message for ${emailAddress} (key: ${generatePersistedSlackMessageKey(cl, branchArg, channelId)})`)
+		this.postOrUpdate(cl, branchArg, dm, persistMessage)
 	}
 
 	// 
-	findAll(sourceCl: number, branchArg: BranchArg) {
+	findAll(cl: number, branchArg: BranchArg) {
 		// Get key to search our list of all messages
 		const persistedMessages = this.persistence.get(SLACK_MESSAGES_PERSISTENCE_KEY) as PersistedSlackMessages
-		const conflictKey = generatePersistedSlackMessageKey(sourceCl, branchArg, "") // pass empty string as channel name
+		const conflictKey = generatePersistedSlackMessageKey(cl, branchArg, "") // pass empty string as channel name
 
 		// Filter messages
 		const messages: SlackConflictMessage[] = []
@@ -329,10 +329,10 @@ export class SlackMessages {
 		return {messages, persistedMessages}
 	}
 
-	find(sourceCl: number, branchArg: BranchArg, channel: string) {
+	find(cl: number, branchArg: BranchArg, channel: string) {
 		// Get key to search our list of all messages
 		const persistedMessages = this.persistence.get(SLACK_MESSAGES_PERSISTENCE_KEY) as PersistedSlackMessages
-		const conflictKey = generatePersistedSlackMessageKey(sourceCl, branchArg, channel)
+		const conflictKey = generatePersistedSlackMessageKey(cl, branchArg, channel)
 		return {messageRecord: persistedMessages[conflictKey], conflictKey, persistedMessages}
 	}
 
@@ -481,45 +481,77 @@ export class BotNotifications implements BotEventHandler {
 			message.footer = blockage.failure.summary
 		}
 
+		let messagesToPost = []
+		let channelsToPostTo = []
+		let usersToInvite: Set<string> = new Set()
+
+		if (userEmail) {
+			usersToInvite.add(userEmail)
+		}
+
 		if (targetBranch) {
 			const additionalChannelInfo = this.additionalBlockChannelIds.get(sourceBranch.upperName + '|' + targetBranch.upperName)
 			if (additionalChannelInfo) {
 				const [sideChannel, postOnlyToSideChannel] = additionalChannelInfo
 				if (!postOnlyToSideChannel) {
-					if (userEmail) {
-						this.slackMessages.addUserToChannel(userEmail, message.channel)
-					}
-					this.slackMessages.postOrUpdate(changeInfo.source_cl, targetBranch.name, message)
-				}
-				if (userEmail) {
-					this.slackMessages.addUserToChannel(userEmail, sideChannel)
+					channelsToPostTo.push(message.channel)
 				}
-				this.slackMessages.postOrUpdate(changeInfo.source_cl, targetBranch.name, { ...message, channel: sideChannel })
+				channelsToPostTo.push(sideChannel)
 			}
 			else {
-				if (userEmail) {
-					this.slackMessages.addUserToChannel(userEmail, message.channel)
+				channelsToPostTo.push(message.channel)
+			}
+
+			messagesToPost.push({ message })
+
+			if (blockage.failure.kind === 'Exclusive check-out') {
+				const exclusiveLockUsers = blockage.failure.additionalInfo as any[]			
+				let text = ''
+
+				for (const exclusiveLockUser of exclusiveLockUsers) {
+					text += `@${exclusiveLockUser.user} `
+					const exclusiveLockUserEmail = await exclusiveLockUser.userEmail
+					if (exclusiveLockUserEmail) {
+						usersToInvite.add(exclusiveLockUserEmail)
+					}
 				}
-				this.slackMessages.postOrUpdate(changeInfo.source_cl, targetBranch.name, message)
+				text += `\n\nPlease unlock the files blocking robomerge or work with ${blockage.owner} to resolve the conflict`
+				
+				const exclusiveCheckoutMessage: SlackMessage = {
+					text,
+					style: SlackMessageStyles.DANGER,
+					channel: message.channel,
+					mrkdwn: true
+				}
+				messagesToPost.push({ message: exclusiveCheckoutMessage, reply: true})
 			}
 		}
 		else {
+			channelsToPostTo.push(message.channel)
+			messagesToPost.push({ message })
 
-			if (userEmail) {
-				this.slackMessages.addUserToChannel(userEmail, message.channel)
-			}
-
-			// separate message for syntax errors (other blockages have a target branch)
-			this.slackMessages.postOrUpdate(changeInfo.source_cl, blockage.failure.kind, message)
-
-			let syntaxErrorMessage: SlackMessage = {
+			const syntaxErrorMessage: SlackMessage = {
 				text: blockage.failure.description,
 				style: SlackMessageStyles.DANGER,
 				channel: message.channel,
 				mrkdwn: false
 			}
-			
-			this.slackMessages.postReply(changeInfo.source_cl, blockage.failure.kind, syntaxErrorMessage)
+			messagesToPost.push({ message: syntaxErrorMessage, reply: true})		
+		}
+
+		const branchArg = targetBranch ? targetBranch.name : blockage.failure.kind
+		for (const channel of channelsToPostTo) {
+			for (const user of usersToInvite) {
+				this.slackMessages.addUserToChannel(user, channel)
+			}
+			for (const messageToPost of messagesToPost) {
+				if (messageToPost.reply) {
+					await this.slackMessages.postReply(changeInfo.cl, branchArg, { ...messageToPost.message, channel })
+				}
+				else {
+					await this.slackMessages.postOrUpdate(changeInfo.cl, branchArg, { ...messageToPost.message, channel })
+				}
+			}
 		}
 
 		// Post message to owner in DM
@@ -537,7 +569,7 @@ export class BotNotifications implements BotEventHandler {
 			else {
 				const dmText = `Your change (${makeClLink(changeInfo.source_cl)}) ` +
 					`hit '${issue}' while merging from *${sourceBranch.name}* to *${targetBranch.name}*.\n\n` +
-					'`' + blockage.change.description.substr(0, 80) + '`\n\n' +
+					'`' + blockage.change.description.substring(0, 80) + '`\n\n' +
 					"*_Resolving this blockage is time sensitive._ Please select one of the following:*"
 				
 				const urls = this.blockageUrlGenerator(blockage)
@@ -547,10 +579,10 @@ export class BotNotifications implements BotEventHandler {
 					throw error
 				}
 
-				dm = this.makeSlackDirectMessage(dmText, changeInfo.source_cl, cl, targetBranch.name, urls)
+				dm = this.makeSlackDirectMessage(dmText, changeInfo.cl, cl, targetBranch.name, urls)
 			}
 
-			this.slackMessages.postDM(userEmail, changeInfo.source_cl, targetBranch, dm)
+			this.slackMessages.postDM(userEmail, changeInfo.cl, targetBranch, dm)
 		}
 	}
 
@@ -561,10 +593,10 @@ export class BotNotifications implements BotEventHandler {
 			if (info.acknowledger) {
 				// hard code to Epic (!) Standard/Daylight Time
 				const suffix = info.acknowledgedAt ? ' at ' + info.acknowledgedAt.toLocaleTimeString('en-US', EPIC_TIME_OPTIONS) : ''
-				this.tryAddFieldToChannelMessages(info.sourceCl, targetKey, {title, value: info.acknowledger + suffix, short: true})
+				this.tryAddFieldToChannelMessages(info.cl, targetKey, {title, value: info.acknowledger + suffix, short: true})
 			}
 			else {
-				this.tryRemoveFieldFromChannelMessages(info.sourceCl, targetKey, title)
+				this.tryRemoveFieldFromChannelMessages(info.cl, targetKey, title)
 			}
 		}
 	}
@@ -597,10 +629,10 @@ export class BotNotifications implements BotEventHandler {
 
 			const messageText = formatResolution(info)
 			const targetKey = info.targetBranchName || info.kind
-			if (!this.updateMessagesAfterUnblock(info.sourceCl, targetKey, '', messageStyle, messageText, newClDesc)) {
+			if (!this.updateMessagesAfterUnblock(info.cl, targetKey, '', messageStyle, messageText, newClDesc)) {
 				this.botNotificationsLogger.warn(`Conflict message not found to update (${info.blockedBranchName} -> ${targetKey} CL#${info.sourceCl})`)
 				const message = this.makeSlackChannelMessage('', messageText, messageStyle, makeClLink(info.cl), info.blockedBranchName, info.targetBranchName, info.author)
-				this.slackMessages.postOrUpdate(info.sourceCl, targetKey, message)
+				this.slackMessages.postOrUpdate(info.cl, targetKey, message)
 			}
 		}
 	}
@@ -797,10 +829,10 @@ export class BotNotifications implements BotEventHandler {
 	
 
 	/** Pre-condition: this.slackMessages must be valid */
-	private updateMessagesAfterUnblock(sourceCl: number, targetBranch: BranchArg, newTitle: string,
+	private updateMessagesAfterUnblock(cl: number, targetBranch: BranchArg, newTitle: string,
 										newStyle: SlackMessageStyles, newText: string, newClDesc?: string) {
 		// Find all messages relating to CL and branch
-		const findResult = this.slackMessages!.findAll(sourceCl, targetBranch)
+		const findResult = this.slackMessages!.findAll(cl, targetBranch)
 
 		if (findResult.messages.length == 0) {
 			return false
@@ -844,7 +876,7 @@ export class BotNotifications implements BotEventHandler {
 
 				// Hacky: If we remove attachments, we'll no longer have a link to the CL in the message.
 				// UE-72320 - Add in a link to the original changelist
-				message.text = newText.replace('change', `change (${makeClLink(sourceCl)})`)
+				message.text = newText.replace('change', `change (${makeClLink(cl)})`)
 			}
 			else {
 				message.text = newText
@@ -864,9 +896,9 @@ export class BotNotifications implements BotEventHandler {
 	}
 
 	/** Pre-condition: this.slackMessages must be valid */
-	private tryAddFieldToChannelMessages(sourceCl: number, targetBranch: BranchArg, newField: SlackMessageField) {
+	private tryAddFieldToChannelMessages(cl: number, targetBranch: BranchArg, newField: SlackMessageField) {
 		// Find all messages relating to CL and branch
-		const findResult = this.slackMessages!.findAll(sourceCl, targetBranch)
+		const findResult = this.slackMessages!.findAll(cl, targetBranch)
 
 		for (const messageRecord of findResult.messages) {
 			const message = messageRecord.messageOpts
@@ -886,9 +918,9 @@ export class BotNotifications implements BotEventHandler {
 	}
 
 	/** Pre-condition: this.slackMessages must be valid */
-	private tryRemoveFieldFromChannelMessages(sourceCl: number, targetBranch: BranchArg, fieldTitle: string) {
+	private tryRemoveFieldFromChannelMessages(cl: number, targetBranch: BranchArg, fieldTitle: string) {
 		// Find all messages relating to CL and branch
-		const findResult = this.slackMessages!.findAll(sourceCl, targetBranch)
+		const findResult = this.slackMessages!.findAll(cl, targetBranch)
 
 		for (const messageRecord of findResult.messages) {
 			const message = messageRecord.messageOpts
