Fix previewing failing when bot has no alias

diff --git a/Engine/Extras/RoboMerge/v3/public/js/flow.js b/Engine/Extras/RoboMerge/v3/public/js/flow.js
--- a/Engine/Extras/RoboMerge/v3/public/js/flow.js
+++ b/Engine/Extras/RoboMerge/v3/public/js/flow.js
@@ -80,7 +80,7 @@ class Graph {
     addNodesAndEdges(grouped) {
         for (const branchStatus of this.branchList) {
             const botTitle = grouped
-                                ? this.options.aliases.get(branchStatus.bot).length > 0
+                                ? this.options.aliases && this.options.aliases.get(branchStatus.bot).length > 0
                                     ? branchStatus.bot + ' (' + this.options.aliases.get(branchStatus.bot).join(", ") + ')'
                                     : branchStatus.bot
                                 : null;
