PR #9013: RoboMerge: Fix Slack API issues (Contributed by [REDACTED_PERSON])

diff --git a/Engine/Extras/RoboMerge/v3/src/robo/slack.ts b/Engine/Extras/RoboMerge/v3/src/robo/slack.ts
--- a/Engine/Extras/RoboMerge/v3/src/robo/slack.ts
+++ b/Engine/Extras/RoboMerge/v3/src/robo/slack.ts
@@ -111,7 +111,7 @@ export class Slack {
 		if (users instanceof Array) {
 			users = users.join(',')
 		}
-		return (await this.post('conversations.open', {token: this.channel.botToken, users})).channel.id
+		return (await this.post('conversations.open', {users})).channel.id
 	}
 
 	/*private*/ async post(command: string, args: any, canFail? : boolean) {
