Reload workspaces if they have been unloaded

diff --git a/Engine/Extras/RoboMerge/v3/src/common/p4util.ts b/Engine/Extras/RoboMerge/v3/src/common/p4util.ts
--- a/Engine/Extras/RoboMerge/v3/src/common/p4util.ts
+++ b/Engine/Extras/RoboMerge/v3/src/common/p4util.ts
@@ -223,7 +223,7 @@ export async function cleanWorkspaces(p4: PerforceContext, workspaces: [string,
 }
 
 export async function getWorkspacesForUser(p4: PerforceContext, user: string, edgeServerAddress?: string) {
-	return (await p4.find_workspaces(user, edgeServerAddress))
+	return (await p4.find_workspaces(user, {edgeServerAddress}))
 		.filter(ws => !USER_WORKSPACE_EXCLUDE_PATTERNS.some(entry => ws.client.match(entry)))
 }
 
diff --git a/Engine/Extras/RoboMerge/v3/src/common/perforce.ts b/Engine/Extras/RoboMerge/v3/src/common/perforce.ts
--- a/Engine/Extras/RoboMerge/v3/src/common/perforce.ts
+++ b/Engine/Extras/RoboMerge/v3/src/common/perforce.ts
@@ -248,8 +248,9 @@ export interface Workspace {
 export type RoboWorkspace = Workspace | string | null;
 
 export interface ClientSpec {
-	client: string;
+	client: string
 	Stream?: string
+	IsUnloaded?: boolean
 }
 
 export type StreamSpec = {
@@ -563,12 +564,14 @@ export class PerforceContext {
 
 	// find a workspace for the given user
 	// output format is an array of workspace names
-	async find_workspaces(user?: string, edgeServerAddress?: string) {
-		// -a to include workspaces on edge servers
+	async find_workspaces(user?: string, options?: {edgeServerAddress?: string, includeUnloaded?: boolean}) {
 
+		const edgeServerAddress = options && options.edgeServerAddress
+		const includeUnloaded = options && options.includeUnloaded
 
 		let args = ['clients', '-u', user || this.username]
 
+		// -a to include workspaces on edge servers
 		if (!edgeServerAddress) {
 			// find all
 			args.push('-a')
@@ -579,13 +582,21 @@ export class PerforceContext {
 			opts.edgeServerAddress = edgeServerAddress
 		}
 
-		let parsedClients = await this._execP4Ztag(null, args, opts);
+		let parsedLoadedClients = this._execP4Ztag(null, args, opts);
+		let parsedUnloadedClients = (includeUnloaded ? this._execP4Ztag(null, [...args, '-U'], opts) : null)
 		let workspaces = [];
-		for (let clientDef of parsedClients) {
+		for (let clientDef of await parsedLoadedClients) {
 			if (clientDef.client) {
 				workspaces.push(clientDef);
 			}
 		}
+		if (includeUnloaded) {
+			for (let clientDef of await parsedUnloadedClients!) {
+				if (clientDef.client) {
+					workspaces.push(clientDef);
+				}
+			}
+		}
 		return workspaces as ClientSpec[];
 	}
 
@@ -606,6 +617,10 @@ export class PerforceContext {
 		return null
 	}
 
+	async reloadWorkspace(workspaceName: string) {
+		return this._execP4Ztag(null, ['reload', '-c', workspaceName]);
+	}
+
 	getEdgeServerAddress(serverId: string): Promise<string> {
 		return this._execP4(null, ['-ztag', '-F', '%Address%', 'server', '-o', serverId])
 	}
diff --git a/Engine/Extras/RoboMerge/v3/src/robo/edgebot.ts b/Engine/Extras/RoboMerge/v3/src/robo/edgebot.ts
--- a/Engine/Extras/RoboMerge/v3/src/robo/edgebot.ts
+++ b/Engine/Extras/RoboMerge/v3/src/robo/edgebot.ts
@@ -485,7 +485,7 @@ class EdgeBotImpl extends PerforceStatefulBot {
 			}
 
 			// do we already have a workspace? (_initBranchWorkspacesForAllBots, _getExistingWorkspaces)
-			const existingWorkspaceInfos = await this.p4.find_workspaces(undefined, edgeServer.address)
+			const existingWorkspaceInfos = await this.p4.find_workspaces(undefined, {edgeServerAddress: edgeServer.address, includeUnloaded: true})
 			if (existingWorkspaceInfos.map(ws => ws.client).indexOf(info.targetWorkspaceOverride) >= 0) {
 				await p4util.cleanWorkspaces(this.p4, [[info.targetWorkspaceOverride, target.branch.rootPath]], edgeServer.address)
 			}
diff --git a/Engine/Extras/RoboMerge/v3/src/robo/robo.ts b/Engine/Extras/RoboMerge/v3/src/robo/robo.ts
--- a/Engine/Extras/RoboMerge/v3/src/robo/robo.ts
+++ b/Engine/Extras/RoboMerge/v3/src/robo/robo.ts
@@ -220,13 +220,14 @@ export class RoboMerge {
 }
 
 async function _getExistingWorkspaces(user?: string) {
-	const existingWorkspacesObjects = await robo.p4.find_workspaces(user);
-	return new Set<string>(existingWorkspacesObjects.map((ws: ClientSpec) => ws.client));
+	const existingWorkspacesObjects = await robo.p4.find_workspaces(user, {includeUnloaded: true});
+	return new Map<string, ClientSpec>(existingWorkspacesObjects.map((ws: ClientSpec) => [ws.client, ws]));
 }
 
-async function _initWorkspacesForGraphBot(graphBot: GraphBot, existingWorkspaces: Set<string>, logger: ContextualLogger) {
+async function _initWorkspacesForGraphBot(graphBot: GraphBot, existingWorkspaces: Map<string, ClientSpec>, logger: ContextualLogger) {
 	// name and depot root pair
 	const workspacesToReset: [string, string][] = []
+	let reloadedWorkspaces = []
 
 	for (const branch of graphBot.branchGraph.branches) {
 		if (branch.workspace !== null) {
@@ -252,7 +253,11 @@ async function _initWorkspacesForGraphBot(graphBot: GraphBot, existingWorkspaces
 		}
 
 		// see if we already have this workspace
-		if (existingWorkspaces.has(ws)) {
+		const clientSpec = existingWorkspaces.get(ws)
+		if (clientSpec) {
+			if (clientSpec.IsUnloaded) {
+				reloadedWorkspaces.push(robo.p4.reloadWorkspace(ws))
+			}
 			workspacesToReset.push([ws, branch.rootPath])
 		}
 		else {
@@ -282,6 +287,7 @@ async function _initWorkspacesForGraphBot(graphBot: GraphBot, existingWorkspaces
 		}
 	}
 
+	Promise.all(reloadedWorkspaces)
 	if (workspacesToReset.length > 0) {
 		logger.info('The following workspaces already exist and will be reset: ' + workspacesToReset.join(', '))
 		await p4util.cleanWorkspaces(robo.p4, workspacesToReset)
