Filter swarm workspaces

diff --git a/Engine/Extras/RoboMerge/v3/src/common/p4util.ts b/Engine/Extras/RoboMerge/v3/src/common/p4util.ts
--- a/Engine/Extras/RoboMerge/v3/src/common/p4util.ts
+++ b/Engine/Extras/RoboMerge/v3/src/common/p4util.ts
@@ -4,7 +4,8 @@ import { ContextualLogger } from './logger';
 import { Change, coercePerforceWorkspace, OpenedFileRecord, PerforceContext, RoboWorkspace } from './perforce';
 
 const USER_WORKSPACE_EXCLUDE_PATTERNS: (RegExp | string)[] = [
-	'horde-p4bridge-'
+	'horde-p4bridge-',
+	'swarm-'
 ]
 
 ///////////////////
