Robomerge can now invite external users to -ext channels

diff --git a/Engine/Extras/RoboMerge/v3/src/robo/notifications.ts b/Engine/Extras/RoboMerge/v3/src/robo/notifications.ts
--- a/Engine/Extras/RoboMerge/v3/src/robo/notifications.ts
+++ b/Engine/Extras/RoboMerge/v3/src/robo/notifications.ts
@@ -76,7 +76,7 @@ export async function postToRobomergeAlerts(message: string) {
 
 export async function postMessageToChannel(message: string, channel: string, style: SlackMessageStyles = SlackMessageStyles.GOOD) {
 	if (SLACK_TOKENS.bot) {
-		return new Slack({id: channel, botToken: SLACK_TOKENS.bot}, args.slackDomain).postMessageToDefaultChannel({
+		return new Slack({id: channel, botToken: SLACK_TOKENS.bot, userToken: SLACK_TOKENS.user}, args.slackDomain).postMessageToDefaultChannel({
 			text: message,
 			style,
 			channel,
@@ -354,13 +354,34 @@ export class SlackMessages {
 		.catch(err => console.error('Error posting non-conflict to Slack! ' + err.toString()))
 	}
 
-	async addUserToChannel(emailAddress: string, channel: string)
+	async addUserToChannel(emailAddress: string, channel: string, externalUser? :boolean)
 	{
 		const user = await this.getSlackUser(emailAddress)
 		if (user) {
-			const result = await this.slack.addUserToChannel(user, channel)
-			if (!result.ok && result.error !== "already_in_channel") {
-				this.smLogger.error(`Error inviting ${emailAddress} (${user}) to channel ${channel}: ${result.error}`)
+			const result = await this.slack.addUserToChannel(user, channel, externalUser)
+			if (!result.ok) {
+				if (result.error === "already_in_channel") { /* this is expected and fine */ }
+				else if (result.error === "failed_for_some_users") { 
+					if (result.failed_user_ids.get(user) === "unable_to_add_user_to_public_channel") {
+						/* this is expected and fine */ 
+					}
+					else {
+						this.smLogger.error(`Error inviting ${emailAddress} (${user}) to channel ${channel}: ${result.error}\n${result.failed_user_ids}`)
+					}
+				}
+				else if (result.error === "user_is_restricted") {
+					// This an external user so we need to use the admin invite and validate that the channel ends in -ext
+					const channelInfo = await this.slack.getChannelInfo(channel)
+					if (channelInfo.channel.name.endsWith('-ext')) {
+						await this.addUserToChannel(emailAddress, channel, true)
+					}
+					else {
+						this.smLogger.error(`Cannot invite external user ${emailAddress} (${user}) to restricted channel ${channel}`)
+					}
+				}
+				else {
+					this.smLogger.error(`Error inviting ${emailAddress} (${user}) to channel ${channel}: ${result.error}`)
+				}
 			}
 		}
 		else {
@@ -947,9 +968,10 @@ export function bindBotNotifications(events: BotEvents, slackChannelOverrides: [
 	let slackMessages
 
 	const botToken = args.devMode && SLACK_DEV_DUMMY_TOKEN || SLACK_TOKENS.bot
+	const userToken = args.devMode && SLACK_DEV_DUMMY_TOKEN || SLACK_TOKENS.user
 	if (botToken && events.botConfig.slackChannel) {
 		logger.info('Enabling Slack messages for ' +  events.botname)
-		slackMessages = new SlackMessages(new Slack({id: events.botConfig.slackChannel, botToken}, args.slackDomain), persistence, logger)
+		slackMessages = new SlackMessages(new Slack({id: events.botConfig.slackChannel, botToken, userToken}, args.slackDomain), persistence, logger)
 	}
 		
 	events.registerHandler(new BotNotifications(events.botname, events.botConfig.slackChannel, externalUrl, blockageUrlGenerator, logger, slackMessages, slackChannelOverrides))
diff --git a/Engine/Extras/RoboMerge/v3/src/robo/slack.ts b/Engine/Extras/RoboMerge/v3/src/robo/slack.ts
--- a/Engine/Extras/RoboMerge/v3/src/robo/slack.ts
+++ b/Engine/Extras/RoboMerge/v3/src/robo/slack.ts
@@ -5,6 +5,7 @@ import * as request from '../common/request';
 export interface SlackChannel {
 	id: string
 	botToken: string
+	userToken: string
 }
 
 export interface SlackMessageField {
@@ -69,8 +70,15 @@ export class Slack {
 	constructor(private channel: SlackChannel, private domain: string) {
 	}
 
-	async addUserToChannel(user: string, channel: string) {
-		return this.post('conversations.invite', {channel: channel, users:user,}, true)
+	async addUserToChannel(user: string, channel: string, externalUser?: boolean) {
+		if (externalUser) {		
+			return this.post_user(this.channel.userToken || this.channel.botToken, 'admin.conversations.invite', {channel_id:channel, user_ids:user}, true)
+		}
+		return this.post('conversations.invite', {channel, users:user}, true)
+	}
+
+	async getChannelInfo(channel:string) {
+		return this.get('conversations.info', {channel})
 	}
 
 	async postMessage(message: SlackMessage) {
@@ -114,11 +122,11 @@ export class Slack {
 		return (await this.post('conversations.open', {users})).channel.id
 	}
 
-	/*private*/ async post(command: string, args: any, canFail? : boolean) {
+	/*private*/ async post_user(userToken: string, command: string, args: any, canFail? : boolean) {
 		const resultJson = await request.post({
 			url: this.domain + '/api/' + command,
 			body: JSON.stringify(args),
-			headers: {Authorization: 'Bearer ' + this.channel.botToken},
+			headers: {Authorization: 'Bearer ' + userToken},
 			contentType: 'application/json; charset=utf-8'
 		})
 		try {
@@ -133,6 +141,10 @@ export class Slack {
 		throw new Error(`${command} generated:\n\t${resultJson}`)
 	}
 
+	/*private*/ async post(command: string, args: any, canFail? : boolean) {
+		return this.post_user(this.channel.botToken, command, args, canFail)
+	}
+
 	/*private*/ async get(command: string, args: any, canFail? : boolean) {
 
 		// erg: why am I always passing a channel?
