Include CL# in reconsider messages that create a shelf
Reorder message to make more sense

diff --git a/Engine/Extras/RoboMerge/v3/src/robo/nodebot.ts b/Engine/Extras/RoboMerge/v3/src/robo/nodebot.ts
--- a/Engine/Extras/RoboMerge/v3/src/robo/nodebot.ts
+++ b/Engine/Extras/RoboMerge/v3/src/robo/nodebot.ts
@@ -1935,10 +1935,7 @@ export class NodeBot extends PerforceStatefulBot implements NodeBotInterface {
 		if (pending.change.userRequest) {
 			this.nodeBotLogger.info(shortMessage + ` (reconsider triggered by ${pending.change.owner}):\n${failure.description}.`)
 
-			let message = `While reconsidering CL#${pending.change.cl}\n`
-			if (postIntegrate && pending.change.targetWorkspaceForShelf) {
-				message += `Shelved in workspace ${pending.change.targetWorkspaceForShelf}\n`
-			}
+			let message = `While reconsidering CL#${pending.change.cl}`
 
 			if (failure.summary) {
 				message += '\n' + failure.summary 
@@ -1948,6 +1945,10 @@ export class NodeBot extends PerforceStatefulBot implements NodeBotInterface {
 				message += '\n' + failure.description
 			}
 
+			if (postIntegrate && pending.change.targetWorkspaceForShelf) {
+				message += `\n\nShelved ${pending.newCl} in workspace ${pending.change.targetWorkspaceForShelf}`
+			}
+
 			if (this.slackMessages) {
 
 				let dm: SlackMessage = {
