Add mechanism to disable mail
Disable mail during functional tests

diff --git a/Engine/Extras/RoboMerge/v3/run-tests.sh b/Engine/Extras/RoboMerge/v3/run-tests.sh
--- a/Engine/Extras/RoboMerge/v3/run-tests.sh
+++ b/Engine/Extras/RoboMerge/v3/run-tests.sh
@@ -75,6 +75,7 @@ docker run -d -p 8877:8877 -p 8811:8811 -p 25:25 -h robomerge_functtest --name r
     -e ROBO_DEV_MODE_USER=testuser1  \
     -e P4PORT=p4docker:1666  \
     -e ROBO_BRANCHSPECS_ROOT_PATH=//RoboMergeData/Main  \
+    -e ROBO_NO_MAIL=true \
     -e ROBO_NO_TLS=true  \
     -e ROBO_EXTERNAL_URL=http://robomerge_functtest:8877  \
     -e BOTNAME=ft1,ft2,ft3,ft4,targets  \
diff --git a/Engine/Extras/RoboMerge/v3/src/common/mailer.ts b/Engine/Extras/RoboMerge/v3/src/common/mailer.ts
--- a/Engine/Extras/RoboMerge/v3/src/common/mailer.ts
+++ b/Engine/Extras/RoboMerge/v3/src/common/mailer.ts
@@ -108,7 +108,12 @@ export class Mailer {
 	analytics: Analytics | null;
 	private readonly mailerLogger = new ContextualLogger('Mailer')
 
-	constructor(analytics: Analytics, parentLogger: ContextualLogger) {
+	constructor(analytics: Analytics, parentLogger: ContextualLogger, mailDisabled: boolean) {
+
+		if (mailDisabled) {
+			return
+		}
+
 		this.mailerLogger = parentLogger.createChild('Mail')
 		this.templates = Mailer.readTemplates(this.mailerLogger)
 		this.mailerLogger.info(`Read email templates: [ ${Array.from(this.templates.keys()).join(', ')} ]`)
diff --git a/Engine/Extras/RoboMerge/v3/src/robo/robo.ts b/Engine/Extras/RoboMerge/v3/src/robo/robo.ts
--- a/Engine/Extras/RoboMerge/v3/src/robo/robo.ts
+++ b/Engine/Extras/RoboMerge/v3/src/robo/robo.ts
@@ -94,6 +94,13 @@ const COMMAND_LINE_ARGS: {[param: string]: Arg<any>} = {
 		dflt: false
 	},
 
+	noMail: {
+		match: /^(-noMail)$/,
+		parse: _str => true,
+		env: 'ROBO_NO_MAIL',
+		dflt: false
+	},
+
 	noTLS: {
 		match: /^(-noTLS)$/,
 		parse: _str => true,
@@ -106,6 +113,7 @@ const COMMAND_LINE_ARGS: {[param: string]: Arg<any>} = {
 		env: 'ROBO_VAULT_PATH',
 		dflt: '/vault'
 	},
+
 	devMode: {
 		match: /^(-devMode)$/,
 		parse: _str => true,
@@ -303,7 +311,7 @@ async function _initBranchWorkspacesForAllBots(logger: ContextualLogger) {
 
 // should be called after data directory has been synced, to get latest email template
 function _initMailer(logger: ContextualLogger) {
-	robo.mailer = new Mailer(roboAnalytics!, logger);
+	robo.mailer = new Mailer(roboAnalytics!, logger, args.noMail);
 }
 
 function _checkForAutoPauseBots(branches: Branch[], logger: ContextualLogger) {
