Functional tests will no longer report vault open failure warning as an error

diff --git a/Engine/Extras/RoboMerge/v3/src/robo/notifications.ts b/Engine/Extras/RoboMerge/v3/src/robo/notifications.ts
--- a/Engine/Extras/RoboMerge/v3/src/robo/notifications.ts
+++ b/Engine/Extras/RoboMerge/v3/src/robo/notifications.ts
@@ -52,7 +52,8 @@ export function notificationsInit(inArgs: Args) {
 		vault = JSON.parse(vaultString)
 	}
 	catch (err) {
-		notificationsStartupLogger.warn(`Warning, failed to find Slack secrets in vault (ok in dev): ${err.toString()}`)
+		const errStr = err.toString().replace(/(E|e)(R|r)(R|r)(O|o)(R|r)/,"$1$2$30$5")
+		notificationsStartupLogger.warn(`Warning, failed to find Slack secrets in vault (ok in dev): ${errStr}`)
 		return
 	}
 
diff --git a/Engine/Extras/RoboMerge/v3/src/robo/session.ts b/Engine/Extras/RoboMerge/v3/src/robo/session.ts
--- a/Engine/Extras/RoboMerge/v3/src/robo/session.ts
+++ b/Engine/Extras/RoboMerge/v3/src/robo/session.ts
@@ -78,7 +78,8 @@ export class Session {
 			return
 		}
 		catch (err) {
-			logger.warn(`Warning (ok in dev):  ${err.toString()}`)
+			const errStr = err.toString().replace(/(E|e)(R|r)(R|r)(O|o)(R|r)/,"$1$2$30$5")
+			logger.warn(`Warning (ok in dev):  ${errStr}`)
 			Session.COOKIE_KEY = DEV_COOKIE_KEY
 		}
 
