Clean up dummy slack server to support all commands that we are using and be explicit about the targets/CL for look up instead of parsing slack messages to try and derive the information

diff --git a/Engine/Extras/RoboMerge/v3/functional_tests/framework/src/framework.ts b/Engine/Extras/RoboMerge/v3/functional_tests/framework/src/framework.ts
--- a/Engine/Extras/RoboMerge/v3/functional_tests/framework/src/framework.ts
+++ b/Engine/Extras/RoboMerge/v3/functional_tests/framework/src/framework.ts
@@ -5,6 +5,14 @@ import * as bent from 'bent'
 import * as System from './system';
 import { Change, Perforce, VERBOSE } from './test-perforce';
 
+// This must be kept in sync with dummyslackserver.ts
+class DummyThread {
+	CL: number
+	target: string
+	timestamp: number
+	messages: string[] // the first message is the original post, the rest are replies
+}
+
 const getJson = bent('json')
 
 const colors = require('colors')
@@ -669,12 +677,12 @@ export abstract class FunctionalTest {
 	}
 
 	wasMessagePostedToSlack(channel: string, cl: number, target: string) {
-		const targetName = this.fullBranchName(target)
+		const targetName = this.fullBranchName(target).toUpperCase()
 		return getJson(DUMMY_SLACK_DOMAIN + '/posted/' + channel)
-			.then((perTargetCls: [string, number[]][]) => {
-				for (const [target, cls] of perTargetCls) {
+			.then((perTargetThreads: [string, DummyThread[]][]) => {
+				for (const [target, dummyThreads] of perTargetThreads) {
 					if (target === targetName) {
-						return cls.indexOf(cl) >= 0
+						return dummyThreads.filter(dt => dt.CL == cl).length == 1
 					}
 				}
 				return false
diff --git a/Engine/Extras/RoboMerge/v3/src/common/dummyslackserver.ts b/Engine/Extras/RoboMerge/v3/src/common/dummyslackserver.ts
--- a/Engine/Extras/RoboMerge/v3/src/common/dummyslackserver.ts
+++ b/Engine/Extras/RoboMerge/v3/src/common/dummyslackserver.ts
@@ -3,9 +3,21 @@
 import { AppInterface, Handler, WebRequest } from './webserver'
 import { setDefault } from './helper'
 
-// channel -> target -> changes
-const posted = new Map<string, Map<string, number[]>>()
+// This must be kept in sync with tests.ts
+class DummyThread {
+	CL: number
+	target: string
+	timestamp: number
+	messages: string[] // the first message is the original post, the rest are replies
+}
+
+// The initial messages of any threads
+// channel -> target -> messages for a change
+const posted = new Map<string, Map<string, DummyThread[]>>()
 
+// The contents of each thread
+// channel -> timestamp -> messages
+const threads = new Map<string, Map<number, DummyThread>>
 
 export class DummySlackApp implements AppInterface {
 
@@ -16,40 +28,84 @@ export class DummySlackApp implements AppInterface {
 	post(command: string) {
 		// @todo (maybe) construct app with logger
 
-		if (command.startsWith("conversations.invite")) {
-			return { ok: true }
-		}
-
 		const data = this.req.reqData
-		if (!data) {
-			throw new Error('Nothing to post!')
-		}
-		const clMatch = data.match(/changes\/(\d+)\|\d+/)
-		const edgeMatch = data.match(/"(\w+) -> (\w+)"/)
-		if (clMatch && edgeMatch) {
+		const getData = function() {
+			if (!data) {
+				throw new Error('No data to parse!')
+			}
 			let dataObj: any
 			try {
 				dataObj = JSON.parse(data)
 			}
 			catch (exc) {
 				console.log('Non-JSON sent to dummy-slack: ', data)
-				return {}
+				return { error: "Non-JSON sent to dummy-slack" }
 			}
+			return { dataObj }
+		}
 
-			setDefault(
-				setDefault(posted, dataObj.channel, new Map),
-				edgeMatch[2], []
-			).push(parseInt(clMatch[1]))
+		if (command === "conversations.invite") {
+			return { ok: true }
 		}
-		else {
-			console.log('dummy-slack POST: cl or edge not found', command, this.req.reqData)
+		else if (command === "conversations.open") {
+			const dataResult = getData()
+			if (dataResult.error) {
+				return dataResult
+			}
+			const dataObj = dataResult.dataObj
+			return { ok: true, channel: {id: `${dataObj.users}`} }
 		}
+		else if (command === "chat.update") {
+			// Not going to handle updating the stored messages on the dummy server
+			// until we have a reason to do so 
+			return { ok: true }
+		}
+		else if (command === "chat.postMessage") {
+			if (!data) {
+				throw new Error('Nothing to post!')
+			}
+			const dataResult = getData()
+			if (dataResult.error) {
+				return dataResult
+			}
+			const dataObj = dataResult.dataObj
+
+			if (dataObj.thread_ts) {
+				let dummyThread = threads.get(dataObj.channel)?.get(dataObj.thread_ts)!
+				dummyThread.messages.push(dataObj.attachments[0].text)
+				return { ok: true }
+			}
+			else {
+				let cl = dataObj.cl
+				let target = dataObj.target
+				if (cl && target) {
 
-		// look for color for block or resolve?
+					const dummyThread = new DummyThread
+					dummyThread.CL = cl
+					dummyThread.target = target
+					dummyThread.timestamp = Date.now()
+					dummyThread.messages = [dataObj.attachments[0].text]
 
-		// careful, two changes for resolutions, luckily first is what we want
+					setDefault(
+						setDefault(posted, dataObj.channel, new Map),
+						dummyThread.target, []
+					).push(dummyThread)
 
-		return {}
+					setDefault(threads, dataObj.channel, new Map)
+						.set(dummyThread.timestamp, dummyThread)
+
+					return { ok: true, ts: dummyThread.timestamp }
+				}
+				else {
+					console.log('dummy-slack POST: cl or edge not found', command, data)
+					return { error: "cl or edge not found"}
+				}
+			}
+		}
+		else {
+			console.error(`dummy-slack POST: unsupported command '${command}'`)
+			return { error: "unsupported command"}
+		}
 	}
 
 	@Handler('GET', '/posted/*')
@@ -61,7 +117,7 @@ export class DummySlackApp implements AppInterface {
 	get(command: string) {
 		if (command.startsWith("users.lookupByEmail"))
 		{
-			return { ok: true, user: { id: "" } }
+			return { ok: true, user: { id: `${this.req.url.searchParams.get('email')}` } }
 		}
 		return { ok: false }
 	}
diff --git a/Engine/Extras/RoboMerge/v3/src/robo/notifications.ts b/Engine/Extras/RoboMerge/v3/src/robo/notifications.ts
--- a/Engine/Extras/RoboMerge/v3/src/robo/notifications.ts
+++ b/Engine/Extras/RoboMerge/v3/src/robo/notifications.ts
@@ -247,6 +247,10 @@ export class SlackMessages {
 		}
 		// Otherwise, we will need to create a new one
 		else {
+
+			message.target = resolveBranchArg(branchArg, true)
+			message.cl = cl
+
 			let timestamp
 			try {
 				timestamp = await this.slack.postMessage(message)
diff --git a/Engine/Extras/RoboMerge/v3/src/robo/slack.ts b/Engine/Extras/RoboMerge/v3/src/robo/slack.ts
--- a/Engine/Extras/RoboMerge/v3/src/robo/slack.ts
+++ b/Engine/Extras/RoboMerge/v3/src/robo/slack.ts
@@ -61,10 +61,14 @@ export interface SlackMessage {
 	channel: string
 	// Allows Markdown formatting in messages
 	mrkdwn: boolean
+
+	// For the dummy server implementation we store additional information
+	cl?: number
+	target?: string
 }
 
 
-const MAIN_MESSAGE_FIELDS = new Set(['username', 'icon_emoji', 'channel']);
+const MAIN_MESSAGE_FIELDS = new Set(['username', 'icon_emoji', 'channel', 'target', 'cl']);
 
 export class Slack {
 	constructor(private channel: SlackChannel, private domain: string) {
