Use slack user ID in message text instead of just [REDACTED_HANDLE] when possible

diff --git a/Engine/Extras/RoboMerge/v3/src/robo/nodebot.ts b/Engine/Extras/RoboMerge/v3/src/robo/nodebot.ts
--- a/Engine/Extras/RoboMerge/v3/src/robo/nodebot.ts
+++ b/Engine/Extras/RoboMerge/v3/src/robo/nodebot.ts
@@ -2068,9 +2068,14 @@ export class NodeBot extends PerforceStatefulBot implements NodeBotInterface {
 									triager = `<@${user}>`
 								}
 								else {
+									triager = `@${triager}`
 									this.nodeBotLogger.error(`Unable to look up triager from ${triager} (${emailAddress})`)
 								}
 							}
+							else {
+								triager = `@${triager}`
+								this.nodeBotLogger.error(`Unable to look up email address for triager ${triager}`)
+							}
 						}
 
 						const nagMessage: SlackMessage = {
diff --git a/Engine/Extras/RoboMerge/v3/src/robo/notifications.ts b/Engine/Extras/RoboMerge/v3/src/robo/notifications.ts
--- a/Engine/Extras/RoboMerge/v3/src/robo/notifications.ts
+++ b/Engine/Extras/RoboMerge/v3/src/robo/notifications.ts
@@ -484,7 +484,8 @@ export class BotNotifications implements BotEventHandler {
 		const issue = blockage.failure.kind.toLowerCase()
 
 		const userEmail = await blockage.ownerEmail
-		const channelPing = `@${blockage.owner}`
+		const slackUser = userEmail ? await this.slackMessages.getSlackUser(userEmail) : null
+		const channelPing = slackUser ? `<@${slackUser}>` : `@${blockage.owner}`
 
 		const isBotUser = isUserAKnownBot(blockage.owner)
 		const text =
@@ -535,11 +536,19 @@ export class BotNotifications implements BotEventHandler {
 				let text = ''
 
 				for (const exclusiveLockUser of exclusiveLockUsers) {
-					text += `@${exclusiveLockUser.user} `
+					let exclusiveLockSlackUser
 					const exclusiveLockUserEmail = await exclusiveLockUser.userEmail
 					if (exclusiveLockUserEmail) {
+						exclusiveLockSlackUser = await this.slackMessages.getSlackUser(exclusiveLockUserEmail)
+						if (exclusiveLockSlackUser) {
+							exclusiveLockSlackUser = `<@${exclusiveLockSlackUser}> `
+						}
 						usersToInvite.add(exclusiveLockUserEmail)
 					}
+					if (!exclusiveLockSlackUser) {
+						exclusiveLockSlackUser = `@${exclusiveLockUser.user} `
+					}
+					text += exclusiveLockSlackUser
 				}
 				text += `\n\nPlease unlock the files blocking robomerge or work with ${blockage.owner} to resolve the conflict`
 				
