Clean up devMode specification so it is possible to set the env var false
Explicit arg to run a dev server while talking to slack
Fix a slack emoji specification

diff --git a/Engine/Extras/RoboMerge/v3/src/robo/notifications.ts b/Engine/Extras/RoboMerge/v3/src/robo/notifications.ts
--- a/Engine/Extras/RoboMerge/v3/src/robo/notifications.ts
+++ b/Engine/Extras/RoboMerge/v3/src/robo/notifications.ts
@@ -187,8 +187,9 @@ function formatResolution(info: PersistentConflict) {
 	if (info.timeTakenToResolveSeconds) {
 		// add time emojis!
 		if (info.timeTakenToResolveSeconds < 2*60) {
-			const poke = Random.choose(['eevee', 'espeon', 'sylveon', 'flareon', 'jolteon', 'leafeon', 'glaceon', 'umbreon', 'vaporeon'])
-			message += ` :${poke}_run:`
+			const emote = Random.choose(['eevee_run', 'espeon_run', 'sylveon_run', 'flareon_run',
+				'jolteon_run', 'leafeon_run', 'glaceon_run', 'umbreon_run', 'vaporeon_run', 'sonic_run', 'fallguy_run'])
+			message += ` :${emote}:`
 		}
 		else if (info.timeTakenToResolveSeconds < 10*60) {
 			message += ' :+1:'
@@ -981,8 +982,8 @@ export function bindBotNotifications(events: BotEvents, slackChannelOverrides: [
 
 	let slackMessages
 
-	const botToken = args.devMode && SLACK_DEV_DUMMY_TOKEN || SLACK_TOKENS.bot
-	const userToken = args.devMode && SLACK_DEV_DUMMY_TOKEN || SLACK_TOKENS.user
+	const botToken = args.devMode && !args.useSlackInDev && SLACK_DEV_DUMMY_TOKEN || SLACK_TOKENS.bot
+	const userToken = args.devMode && !args.useSlackInDev && SLACK_DEV_DUMMY_TOKEN || SLACK_TOKENS.user
 	if (botToken && events.botConfig.slackChannel) {
 		logger.info('Enabling Slack messages for ' +  events.botname)
 		slackMessages = new SlackMessages(new Slack({id: events.botConfig.slackChannel, botToken, userToken}, args.slackDomain), persistence, logger)
diff --git a/Engine/Extras/RoboMerge/v3/src/robo/robo.ts b/Engine/Extras/RoboMerge/v3/src/robo/robo.ts
--- a/Engine/Extras/RoboMerge/v3/src/robo/robo.ts
+++ b/Engine/Extras/RoboMerge/v3/src/robo/robo.ts
@@ -108,6 +108,13 @@ const COMMAND_LINE_ARGS: {[param: string]: Arg<any>} = {
 		dflt: false
 	},
 
+	useSlackInDev: {
+		match: /^(-useSlackInDev)$/,
+		parse: _str => true,
+		env: 'ROBO_USE_SLACK_IN_DEV',
+		dflt: false
+	},
+
 	vault: {
 		match: /^-vault_path=(.+)$/,
 		env: 'ROBO_VAULT_PATH',
@@ -116,7 +123,7 @@ const COMMAND_LINE_ARGS: {[param: string]: Arg<any>} = {
 
 	devMode: {
 		match: /^(-devMode)$/,
-		parse: _str => true,
+		parse: str => str === "false" ? false : true,
 		env: 'ROBO_DEV_MODE',
 		dflt: false
 	},
diff --git a/Engine/Extras/RoboMerge/v3/src/robo/roboserver.ts b/Engine/Extras/RoboMerge/v3/src/robo/roboserver.ts
--- a/Engine/Extras/RoboMerge/v3/src/robo/roboserver.ts
+++ b/Engine/Extras/RoboMerge/v3/src/robo/roboserver.ts
@@ -26,7 +26,7 @@ let ENVIRONMENT: {[param: string]: any}
 	const COMMAND_LINE_ARGS: {[param: string]: (Arg<any>)} = {
 		devMode: {
 			match: /^(-devMode)$/,
-			parse: _str => true,
+			parse: str => str === "false" ? false : true,
 			env: 'ROBO_DEV_MODE',
 			dflt: false
 		},
