Gracefully handle when perforce is up but the branchspecs haven't been populated yet (primarily for functional test start up)

diff --git a/Engine/Extras/RoboMerge/v3/src/robo/robo.ts b/Engine/Extras/RoboMerge/v3/src/robo/robo.ts
--- a/Engine/Extras/RoboMerge/v3/src/robo/robo.ts
+++ b/Engine/Extras/RoboMerge/v3/src/robo/robo.ts
@@ -395,6 +395,22 @@ async function init(logger: ContextualLogger) {
 		if (workspace.length === 0) {
 			logger.info("Cannot find branch spec workspace " + args.branchSpecsWorkspace + 
 						", creating a new one.")
+
+			while (true) {
+				try {
+					const streams = await robo.p4.streams();
+					if (streams.has(args.branchSpecsRootPath)) {
+						break
+					}
+				}
+				catch (err) {
+				}
+
+				const timeout = 5.0;
+				logger.info(`Will check again in ${timeout} sec...`);
+				await _setTimeout(timeout*1000);
+			}
+
 			await robo.p4.newBranchSpecWorkspace(autoUpdaterConfig.workspace, args.branchSpecsRootPath)
 		}
 
