Filter functional test dist folder

diff --git a/Engine/Extras/RoboMerge/v3/.vscode/settings.json b/Engine/Extras/RoboMerge/v3/.vscode/settings.json
--- a/Engine/Extras/RoboMerge/v3/.vscode/settings.json
+++ b/Engine/Extras/RoboMerge/v3/.vscode/settings.json
@@ -1,5 +1,6 @@
 {
     "search.exclude": {
-        "dist": true
+        "dist": true,
+        "functional_tests/framework/dist": true
     }
 }
\ No newline at end of file
