Approvals enforced regardless of whether the robomerge user is allowed to submit in the stream or not

diff --git a/Engine/Extras/RoboMerge/v3/src/robo/edgebot.ts b/Engine/Extras/RoboMerge/v3/src/robo/edgebot.ts
--- a/Engine/Extras/RoboMerge/v3/src/robo/edgebot.ts
+++ b/Engine/Extras/RoboMerge/v3/src/robo/edgebot.ts
@@ -650,16 +650,30 @@ class EdgeBotImpl extends PerforceStatefulBot {
 		}
 		else
 		{
-			// no conflicts, try to submit
-			const result = await this.submitChangelist(pending, submitRetries)
-			if (result.result !== 'error') {
-				return result
-			}
+			if (this.options.approval) {
+				const approval = this.options.approval
 
-			// todo: lock a target stream in the tests, print error message
-			// or lock a test stream and try to commit
-			// then look for error message in result.message
-			if (!this.options.approval) {
+				// shelve
+				await this.shelveChangelist(pending)
+	
+				// still notify as a 'failure' to trigger normal blockage mechanism
+				failure = { kind: 'Approval required', description: approval.description }
+	
+				// pause this bot
+				this.block(this.createEdgeBlockageInfo(failure, pending))
+				this.sourceNode.findOrCreateBlockage(failure, pending,
+					`Integration of CL#${pending.change.cl} to ${pending.action.branch.name} needs approval: ${approval.description}`,
+					{
+						settings: approval,
+						shelfCl: pending.newCl
+					});
+			}
+			else {
+				// no conflicts, try to submit
+				const result = await this.submitChangelist(pending, submitRetries)
+				if (result.result !== 'error') {
+					return result
+				}
 
 				let summary: string | undefined
 				const match = result.message.match(/.*STDERR:([^]*)STDOUT:/)
@@ -673,28 +687,10 @@ class EdgeBotImpl extends PerforceStatefulBot {
 		}
 
 		// no failure set means the change needs approval
-		if (failure) {
+		if (failure && failure.kind !== 'Approval required') {
 			await this.handlePostIntegrationFailure(failure, pending)
 			await this.sourceNode.handleMergeFailure(failure, pending, true)
 		}
-		else {
-			const approval = this.options.approval!
-
-			// shelve
-			await this.shelveChangelist(pending)
-
-			// still notify as a 'failure' to trigger normal blockage mechanism
-			failure = { kind: 'Approval required', description: approval.description }
-
-			// pause this bot
-			this.block(this.createEdgeBlockageInfo(failure, pending))
-			this.sourceNode.findOrCreateBlockage(failure, pending,
-				`Integration of CL#${pending.change.cl} to ${pending.action.branch.name} needs approval: ${approval.description}`,
-				{
-					settings: approval,
-					shelfCl: pending.newCl
-				});
-		}
 
 		return new EdgeIntegrationDetails('error', `${failure.kind}: ${failure.description}`)
 	}
