Reconsiders on edges requiring approvals now post to the channel

diff --git a/Engine/Extras/RoboMerge/v3/src/robo/edgebot.ts b/Engine/Extras/RoboMerge/v3/src/robo/edgebot.ts
--- a/Engine/Extras/RoboMerge/v3/src/robo/edgebot.ts
+++ b/Engine/Extras/RoboMerge/v3/src/robo/edgebot.ts
@@ -659,14 +659,20 @@ class EdgeBotImpl extends PerforceStatefulBot {
 				// still notify as a 'failure' to trigger normal blockage mechanism
 				failure = { kind: 'Approval required', description: approval.description }
 	
-				// pause this bot
-				this.block(this.createEdgeBlockageInfo(failure, pending))
-				this.sourceNode.findOrCreateBlockage(failure, pending,
-					`Integration of CL#${pending.change.cl} to ${pending.action.branch.name} needs approval: ${approval.description}`,
-					{
-						settings: approval,
-						shelfCl: pending.newCl
-					});
+				const msg = `Integration of CL#${pending.change.cl} to ${pending.action.branch.name} needs approval: ${approval.description}`
+				if (pending.change.userRequest) {
+					postMessageToChannel(msg, approval.channelId, SlackMessageStyles.DANGER)			
+				}
+				else {
+					// pause this bot
+					this.block(this.createEdgeBlockageInfo(failure, pending))
+					this.sourceNode.findOrCreateBlockage(failure, pending,
+						msg,
+						{
+							settings: approval,
+							shelfCl: pending.newCl
+						});
+				}
 			}
 			else {
 				// no conflicts, try to submit
