Fix reconsider approvals messaging

diff --git a/Engine/Extras/RoboMerge/v3/src/robo/edgebot.ts b/Engine/Extras/RoboMerge/v3/src/robo/edgebot.ts
--- a/Engine/Extras/RoboMerge/v3/src/robo/edgebot.ts
+++ b/Engine/Extras/RoboMerge/v3/src/robo/edgebot.ts
@@ -659,15 +659,14 @@ class EdgeBotImpl extends PerforceStatefulBot {
 				// still notify as a 'failure' to trigger normal blockage mechanism
 				failure = { kind: 'Approval required', description: approval.description }
 	
-				const msg = `Integration of CL#${pending.change.cl} to ${pending.action.branch.name} needs approval: ${approval.description}`
 				if (pending.change.userRequest) {
-					postMessageToChannel(msg, approval.channelId, SlackMessageStyles.DANGER)			
+					this.sourceNode.reportApprovalRequired(approval, pending) 
 				}
 				else {
 					// pause this bot
 					this.block(this.createEdgeBlockageInfo(failure, pending))
 					this.sourceNode.findOrCreateBlockage(failure, pending,
-						msg,
+						`Integration of CL#${pending.change.cl} to ${pending.action.branch.name} needs approval: ${approval.description}`,
 						{
 							settings: approval,
 							shelfCl: pending.newCl
diff --git a/Engine/Extras/RoboMerge/v3/src/robo/nodebot.ts b/Engine/Extras/RoboMerge/v3/src/robo/nodebot.ts
--- a/Engine/Extras/RoboMerge/v3/src/robo/nodebot.ts
+++ b/Engine/Extras/RoboMerge/v3/src/robo/nodebot.ts
@@ -13,11 +13,11 @@ import { MergeAction, OperationResult, PendingChange, resolveBranchArg, StompedR
 import { Conflicts } from './conflicts';
 import { EdgeBot, EdgeMergeResults } from './edgebot';
 import { BotEventTriggers } from './events';
-import { SlackMessages } from './notifications';
+import { makeClLink, SlackMessages } from './notifications';
 import { PerforceStatefulBot } from './perforce-stateful-bot';
 import { BlockageNodeOpUrls, OperationUrlHelper } from './roboserver';
 import { Context } from './settings';
-import { SlackMessage, SlackMessageStyles } from './slack';
+import { SlackMessage, SlackMessageField, SlackMessageStyles } from './slack';
 import { PauseState } from './state-interfaces';
 import { newTickJournal, TickJournal } from './tick-journal';
 import { computeTargets, parseDescriptionLines, processOtherBotTargets, getIntegrationOwner, getNodeBotFullName, getNodeBotFullNameForLogging } from './targets';
@@ -1994,6 +1994,29 @@ export class NodeBot extends PerforceStatefulBot implements NodeBotInterface {
 		return true
 	}
 
+	async reportApprovalRequired(approval: ApprovalOptions, pending: PendingChange) {
+		if (this.slackMessages) {
+			const userEmail = await this.p4.getEmail(pending.change.author)
+			const slackUser = userEmail ? await this.slackMessages.getSlackUser(userEmail) : null
+			const channelPing = slackUser ? `<@${slackUser}>` : `@${pending.change.author}`
+
+			const fields: SlackMessageField[] = [
+				{title: 'Change', short: true, value: makeClLink(pending.newCl)}
+			]
+
+			const opts: SlackMessage = { 
+				title:'', 
+				text: `${channelPing}'s change needs to be approved.`, 
+				style: SlackMessageStyles.DANGER, 
+				fields,
+				mrkdwn: true,
+				channel: approval.channelId
+			}
+
+			this.slackMessages.postNonConflictMessage(opts)
+		}
+	}
+
 	private async sendNagNotifications() {
 		let naggedSomeone = false
 		const now = Date.now()
diff --git a/Engine/Extras/RoboMerge/v3/src/robo/notifications.ts b/Engine/Extras/RoboMerge/v3/src/robo/notifications.ts
--- a/Engine/Extras/RoboMerge/v3/src/robo/notifications.ts
+++ b/Engine/Extras/RoboMerge/v3/src/robo/notifications.ts
@@ -423,7 +423,7 @@ export class SlackMessages {
 	}
 }
 
-function makeClLink(cl: number, alias?: string) {
+export function makeClLink(cl: number, alias?: string) {
 	return `<https://p4-swarm.companyname.net/changes/${cl}|${alias ? alias : cl}>`
 }
 
