Possible fix for compile error that does not repro locally

diff --git a/Engine/Extras/RoboMerge/v3/src/common/dummyslackserver.ts b/Engine/Extras/RoboMerge/v3/src/common/dummyslackserver.ts
--- a/Engine/Extras/RoboMerge/v3/src/common/dummyslackserver.ts
+++ b/Engine/Extras/RoboMerge/v3/src/common/dummyslackserver.ts
@@ -17,7 +17,7 @@ const posted = new Map<string, Map<string, DummyThread[]>>()
 
 // The contents of each thread
 // channel -> timestamp -> messages
-const threads = new Map<string, Map<number, DummyThread>>
+const threads = new Map<string, Map<number, DummyThread>>()
 
 export class DummySlackApp implements AppInterface {
 
