Fix interaction with slack API results when admin invite "fails"

diff --git a/Engine/Extras/RoboMerge/v3/src/robo/notifications.ts b/Engine/Extras/RoboMerge/v3/src/robo/notifications.ts
--- a/Engine/Extras/RoboMerge/v3/src/robo/notifications.ts
+++ b/Engine/Extras/RoboMerge/v3/src/robo/notifications.ts
@@ -368,7 +368,7 @@ export class SlackMessages {
 			if (!result.ok) {
 				if (result.error === "already_in_channel") { /* this is expected and fine */ }
 				else if (result.error === "failed_for_some_users") { 
-					if (result.failed_user_ids.get(user) === "unable_to_add_user_to_public_channel") {
+					if (result.failed_user_ids[user] === "unable_to_add_user_to_public_channel") {
 						/* this is expected and fine */ 
 					}
 					else {
