Fix waiting for branch spec stream to exist to create workspace when path to specs is not the root

diff --git a/Engine/Extras/RoboMerge/v3/src/robo/robo.ts b/Engine/Extras/RoboMerge/v3/src/robo/robo.ts
--- a/Engine/Extras/RoboMerge/v3/src/robo/robo.ts
+++ b/Engine/Extras/RoboMerge/v3/src/robo/robo.ts
@@ -397,18 +397,25 @@ async function init(logger: ContextualLogger) {
 						", creating a new one.")
 
 			while (true) {
-				try {
-					const streams = await robo.p4.streams();
-					if (streams.has(args.branchSpecsRootPath)) {
-						break
+				const match = args.branchSpecsRootPath.match(/(\/\/.*?\/.*?)[\/$]/)
+				if (match) {
+					const branchSpecsStream = match[1]
+					try {
+						const streams = await robo.p4.streams();
+						if (streams.has(branchSpecsStream)) {
+							break
+						}
 					}
+					catch (err) {
+					}
+
+					const timeout = 5.0;
+					logger.info(`Will check again in ${timeout} sec...`);
+					await _setTimeout(timeout*1000);
 				}
-				catch (err) {
+				else {
+					logger.warn(`Unable to determine stream from root path '${args.branchSpecsRootPath}'`)
 				}
-
-				const timeout = 5.0;
-				logger.info(`Will check again in ${timeout} sec...`);
-				await _setTimeout(timeout*1000);
 			}
 
 			await robo.p4.newBranchSpecWorkspace(autoUpdaterConfig.workspace, args.branchSpecsRootPath)
