Fix #robomerge-owner not being correctly populated when using edge resolvers

diff --git a/Engine/Extras/RoboMerge/v3/src/robo/edgebot.ts b/Engine/Extras/RoboMerge/v3/src/robo/edgebot.ts
--- a/Engine/Extras/RoboMerge/v3/src/robo/edgebot.ts
+++ b/Engine/Extras/RoboMerge/v3/src/robo/edgebot.ts
@@ -357,7 +357,7 @@ class EdgeBotImpl extends PerforceStatefulBot {
 		description += '\n\n' // description has been trimmed
 
 		// if the owner is specifically overridden (can be by reconsider, resolver or manual tag)
-		const overriddenOwner = getIntegrationOwner(this.branch, info.owner)
+		const overriddenOwner = getIntegrationOwner(this.branch, info.branch, info.owner)
 
 		if (overriddenOwner) {
 			description += `#ROBOMERGE-OWNER: ${overriddenOwner}\n`
diff --git a/Engine/Extras/RoboMerge/v3/src/robo/targets.ts b/Engine/Extras/RoboMerge/v3/src/robo/targets.ts
--- a/Engine/Extras/RoboMerge/v3/src/robo/targets.ts
+++ b/Engine/Extras/RoboMerge/v3/src/robo/targets.ts
@@ -7,9 +7,9 @@ import { ContextualLogger } from '../common/logger';
 
 // const NEW_STUFF = false
 
-export function getIntegrationOwner(targetBranch: Branch, overriddenOwner?: string): string | null
+export function getIntegrationOwner(targetBranch: Branch, sourceBranch: Branch, overriddenOwner?: string): string | null
 export function getIntegrationOwner(pending: PendingChange): string
-export function getIntegrationOwner(arg0: Branch | PendingChange, overriddenOwner?: string) {
+export function getIntegrationOwner(arg0: Branch | PendingChange, sourceBranch?: Branch, overriddenOwner?: string) {
 	// order of priority for owner:
 
 	//  1)	a) Change flagged 'manual', i.e. will create a shelf
@@ -31,28 +31,29 @@ export function getIntegrationOwner(arg0: Branch | PendingChange, overriddenOwne
 		targetBranch = arg0 as Branch
 	}
 
-
 	// Manual requester or edge reconsider
-	if (pending && (
-			pending.action.flags.has('manual') ||
+	if (pending) {
+		if (pending.action.flags.has('manual') ||
 			pending.change.forceCreateAShelf ||
 			pending.change.userRequest == 'edge-reconsider'
-		)) {
-		return pending.change.owner
-	}
+		) {
+			return pending.change.owner
+		}
 
-	const branch = pending ? pending.action.branch : targetBranch
-	const owner = pending ? pending.change.owner : overriddenOwner
+		targetBranch = pending.action.branch
+		sourceBranch = pending.change.branch
+		overriddenOwner = pending.change.owner
+	}
 
-	let resolver = branch!.resolver
-	if (pending) {
-		const edgeprops = pending.change.branch.edgeProperties.get(branch!.upperName)
+	let resolver = targetBranch!.resolver
+	if (sourceBranch) {
+		const edgeprops = sourceBranch.edgeProperties.get(targetBranch!.upperName)
 		if (edgeprops) {
 			resolver = edgeprops.resolver || resolver || null
 		}	
 	}
 
-	return resolver || owner || null
+	return resolver || overriddenOwner || null
 }
 
 export function getNodeBotFullName(botname: string, branchName: string) {
