Fix acknowledgedAt time not getting made back in to a Date up during load from persistent data

diff --git a/Engine/Extras/RoboMerge/v3/src/robo/conflicts.ts b/Engine/Extras/RoboMerge/v3/src/robo/conflicts.ts
--- a/Engine/Extras/RoboMerge/v3/src/robo/conflicts.ts
+++ b/Engine/Extras/RoboMerge/v3/src/robo/conflicts.ts
@@ -108,6 +108,9 @@ export class Conflicts {
 					if (conflict.lastNagTime) {
 						conflict.lastNagTime = new Date(conflict.lastNagTime)
 					}
+					if (conflict.acknowledgedAt) {
+						conflict.acknowledgedAt = new Date(conflict.acknowledgedAt)
+					}
 					conflict.nagCount = conflict.nagCount || 0 // back compat for nag changes
 					this.conflicts.push(conflict)
 				}
