Nag leeway

diff --git a/Engine/Extras/RoboMerge/v3/src/robo/branchdefs.ts b/Engine/Extras/RoboMerge/v3/src/robo/branchdefs.ts
--- a/Engine/Extras/RoboMerge/v3/src/robo/branchdefs.ts
+++ b/Engine/Extras/RoboMerge/v3/src/robo/branchdefs.ts
@@ -20,6 +20,8 @@ export interface BotConfig {
 	globalNotify: string[]
 	nagWhenBlocked: boolean | null
 	nagSchedule: number[] | null
+	nagAcknowledgedSchedule: number[] | null
+	nagAcknowledgedLeeway: number | null
 	triager: string | null
 	checkIntervalSecs: number
 	excludeAuthors: string[]
@@ -50,6 +52,8 @@ const branchBasePrototype = {
 	triager: '' as string | null,
 	nagWhenBlocked: false as boolean | null,
 	nagSchedule: [] as number[] | null,
+	nagAcknowledgedSchedule: [] as number[] | null,
+	nagAcknowledgedLeeway: 0 as number | null,
 	aliases: [''],
 	badgeProject: '' as string | null,
 }
@@ -146,6 +150,8 @@ const edgeOptionFieldsPrototype = {
 	triager: '',
 
 	nagSchedule: [],
+	nagAcknowledgedSchedule: [],
+	nagAcknowledgedLeeway: 0,
 	nagWhenBlocked: true,
 
 	terminal: false, // changes go along terminal edges but no further
@@ -268,6 +274,8 @@ export class BranchDefs {
 			globalNotify: [],
 			triager: null,
 			nagSchedule: null,
+			nagAcknowledgedSchedule: null,
+			nagAcknowledgedLeeway: null,
 			nagWhenBlocked: null,
 			emailOnBlockage: true,
 			checkIntervalSecs: 30.0,
diff --git a/Engine/Extras/RoboMerge/v3/src/robo/branchgraph.ts b/Engine/Extras/RoboMerge/v3/src/robo/branchgraph.ts
--- a/Engine/Extras/RoboMerge/v3/src/robo/branchgraph.ts
+++ b/Engine/Extras/RoboMerge/v3/src/robo/branchgraph.ts
@@ -240,6 +240,8 @@ if (botname === '__TEST__') {
 			resolver: options.resolver || null,
 			triager: options.triager || null,
 			nagSchedule: options.nagSchedule || null,
+			nagAcknowledgedSchedule: options.nagAcknowledgedSchedule || null,
+			nagAcknowledgedLeeway: options.nagAcknowledgedLeeway || null,
 			nagWhenBlocked: options.nagWhenBlocked || null,
 			convertIntegratesToEdits: (options.integrationMethod || this.config.defaultIntegrationMethod) === IntegrationMethod.CONVERT_TO_EDIT,
 			visibility: options.visibility || this.config.visibility,
diff --git a/Engine/Extras/RoboMerge/v3/src/robo/nodebot.ts b/Engine/Extras/RoboMerge/v3/src/robo/nodebot.ts
--- a/Engine/Extras/RoboMerge/v3/src/robo/nodebot.ts
+++ b/Engine/Extras/RoboMerge/v3/src/robo/nodebot.ts
@@ -28,6 +28,8 @@ import { GraphAPI } from '../new/graph';
  **********************************/
 
 const DEFAULT_NAG_SCHEDULE = [10, 20, 30, 60]
+const DEFAULT_ACKNOWLEDGED_NAG_SCHEDULE = [30, 30, 60]
+const DEFAULT_ACKNOWLEDGED_NAG_LEEWAY = 15
 const NAG_EMAIL_MIN_TIME_MINUTES = 60
 const NAG_EMAIL_MIN_TIME_DESCRIPTION = 'an hour'
 const SYNTAX_ERROR_PAUSE_TIMEOUT_SECONDS = 10 * 60
@@ -2055,12 +2057,15 @@ export class NodeBot extends PerforceStatefulBot implements NodeBotInterface {
 
 				if (this.slackMessages) {
 
-					const nagSchedule = getNagProperty("nagSchedule") || DEFAULT_NAG_SCHEDULE
+					const nagSchedule = (conflict.acknowledger && (getNagProperty("nagAcknowledgedSchedule") || DEFAULT_ACKNOWLEDGED_NAG_SCHEDULE))
+											|| getNagProperty("nagSchedule") || DEFAULT_NAG_SCHEDULE
 
 					const nagDelay = conflict.lastNagTime ? ((now - conflict.lastNagTime.getTime()) / 1000) / 60 : ageMinutes
 					const nagWait = nagSchedule[Math.min(conflict.nagCount, nagSchedule.length - 1)]
+					const nagLeeway = getNagProperty("nagAcknowledgedLeeway") || DEFAULT_ACKNOWLEDGED_NAG_LEEWAY
+					const giveLeeway = conflict.acknowledgedAt && (nagLeeway > ((now - conflict.acknowledgedAt.getTime()) / 1000) / 60)
 
-					if (nagDelay > nagWait) {
+					if (nagDelay > nagWait && !giveLeeway) {
 						conflict.nagCount++ 
 						conflict.lastNagTime = new Date(now)
 						naggedSomeone = true
