Fix regex for //depot/stream case

diff --git a/Engine/Extras/RoboMerge/v3/src/robo/robo.ts b/Engine/Extras/RoboMerge/v3/src/robo/robo.ts
--- a/Engine/Extras/RoboMerge/v3/src/robo/robo.ts
+++ b/Engine/Extras/RoboMerge/v3/src/robo/robo.ts
@@ -397,7 +397,7 @@ async function init(logger: ContextualLogger) {
 						", creating a new one.")
 
 			while (true) {
-				const match = args.branchSpecsRootPath.match(/(\/\/.*?\/.*?)[\/$]/)
+				const match = args.branchSpecsRootPath.match(/(\/\/.*?\/.*?)(\/|$)/)
 				if (match) {
 					const branchSpecsStream = match[1]
 					try {
