Fix exclusive locks not reporting correctly when the exclusive lock is for an add

diff --git a/Engine/Extras/RoboMerge/v3/src/robo/edgebot.ts b/Engine/Extras/RoboMerge/v3/src/robo/edgebot.ts
--- a/Engine/Extras/RoboMerge/v3/src/robo/edgebot.ts
+++ b/Engine/Extras/RoboMerge/v3/src/robo/edgebot.ts
@@ -276,20 +276,26 @@ class EdgeBotImpl extends PerforceStatefulBot {
 			this.edgeBotLogger.error(`Integration error: ${errors.length} files, checking first ${MAX_INTEGRATION_ERRORS_TO_ANALYZE}`)
 		}
 		
-		const openedRequests: [RegExpMatchArray, Promise<OpenedFileRecord[]>][] = []
+		const openedRequests: [RegExpMatchArray, Promise<OpenedFileRecord[]>, Promise<OpenedFileRecord[]>][] = []
 		for (const err of errors.slice(0, MAX_INTEGRATION_ERRORS_TO_ANALYZE)) {
 			const match = err.match(EXCLUSIVE_CHECKOUT_REGEX)
 			if (match) {
-				openedRequests.push([match, this.p4.opened(null, match[1] + match[2], true)])
+				openedRequests.push([match, this.p4.opened(null, match[1] + match[2], true), this.p4.opened(null, match[1] + match[2])])
 			}
 		}
 
 		const results: ExclusiveFile[] = []
-		for (const [match, req] of openedRequests) {
-			const recs = await req
+		for (const [match, exclusiveReq, addReq] of openedRequests) {
+			const recs = await exclusiveReq
 			if (recs.length > 0) {
 				// should only be one, since we're looking for exclusive check-out errors
 				results.push({name: match[2], user: recs[0].user})
+			} else {
+				const recs = await addReq
+				if (recs.length > 0) {
+					// should only be one, since we're looking for exclusive check-out errors
+					results.push({name: match[2], user: recs[0].user})
+				}
 			}
 		}
 
