Factor out a more generic editChange so Owner and Description could be modified in a single command

diff --git a/Engine/Extras/RoboMerge/v3/src/common/perforce.ts b/Engine/Extras/RoboMerge/v3/src/common/perforce.ts
--- a/Engine/Extras/RoboMerge/v3/src/common/perforce.ts
+++ b/Engine/Extras/RoboMerge/v3/src/common/perforce.ts
@@ -280,9 +280,11 @@ interface ExecZtagOpts extends ExecOpts {
 	multiline?: boolean;
 }
 
-export interface EditOwnerOpts {
-	newWorkspace?: string;
-	changeSubmitted?: boolean;
+export interface EditChangeOpts {
+	newOwner?: string
+	newWorkspace?: string
+	newDescription?: string
+	changeSubmitted?: boolean
 	edgeServerAddress?: string
 }
 
@@ -1246,46 +1248,51 @@ export class PerforceContext {
 		return result;
 	}
 
-	// change the description on an existing CL
-	// output format is just error or not
-	async editDescription(roboWorkspace: RoboWorkspace, changelist: number, description: string, edgeServerAddress?: string) {
-		const workspace = coercePerforceWorkspace(roboWorkspace);
-
-		// get the current changelist description
-		const output = await this._execP4(workspace, ['change', '-o', changelist.toString()], {edgeServerAddress});
-		// replace the description
-		let new_desc = '\nDescription:\n\t' + this._sanitizeDescription(description);
-		let form = output.replace(/\nDescription:\n(\t[^\n]*\n)*/, new_desc.replace(/\$/g, '$$$$'));
-
-		// run the P4 change command to update
-		this.logger.info("Executing: 'p4 change -i -u' to edit description on CL" + changelist);
-		await this._execP4(workspace, ['change', '-i', '-u'], { stdin: form, quiet: true, edgeServerAddress });
-	}
-
-	// change the owner of an existing CL
+	// update the fields on an existing CL using p4 change
 	// output format is just error or not
-	async editOwner(roboWorkspace: RoboWorkspace, changelist: number, newOwner: string, opts?: EditOwnerOpts) {
+	async editChange(roboWorkspace: RoboWorkspace, changelist: number, opts?: EditChangeOpts) {
 		const workspace = coercePerforceWorkspace(roboWorkspace);
 
 		opts = opts || {}; // optional newWorkspace:string and/or changeSubmitted:boolean
 		// get the current changelist description
-		const output = await this._execP4(workspace, ['change', '-o', changelist.toString()],
+		let form = await this._execP4(workspace, ['change', '-o', changelist.toString()],
 			{edgeServerAddress: opts.edgeServerAddress});
 
-		// replace the description
-		let form = output.replace(/\nUser:\t[^\n]*\n/, `\nUser:\t${newOwner}\n`);
+		if (opts.newOwner) {
+			form = form.replace(/\nUser:\t[^\n]*\n/, `\nUser:\t${opts.newOwner}\n`);
+		}
 		if (opts.newWorkspace) {
 			form = form.replace(/\nClient:\t[^\n]*\n/, `\nClient:\t${opts.newWorkspace}\n`);
 		}
+		if (opts.newDescription) {
+			// replace the description
+			let new_desc = '\nDescription:\n\t' + this._sanitizeDescription(opts.newDescription);
+			form = form.replace(/\nDescription:\n(\t[^\n]*\n)*/, new_desc.replace(/\$/g, '$$$$'));
+		}
 
 		// run the P4 change command to update
 		const changeFlag = opts.changeSubmitted ? '-f' : '-u';
-		this.logger.info(`Executing: 'p4 change -i ${changeFlag}' to edit user/client on CL${changelist}`);
+		this.logger.info(`Executing: 'p4 change -i ${changeFlag}' to edit CL${changelist}`);
 
 		await this._execP4(workspace, ['change', '-i', changeFlag],
 				{ stdin: form, quiet: true, edgeServerAddress: opts.edgeServerAddress })
 	}
 
+	// change the description on an existing CL
+	// output format is just error or not
+	async editDescription(roboWorkspace: RoboWorkspace, changelist: number, description: string, edgeServerAddress?: string) {
+		const opts: EditChangeOpts = {newDescription: description, edgeServerAddress}
+		this.editChange(roboWorkspace, changelist, opts)
+	}
+
+	// change the owner of an existing CL
+	// output format is just error or not
+	async editOwner(roboWorkspace: RoboWorkspace, changelist: number, newOwner: string, opts?: EditChangeOpts) {
+		opts = opts || {}; // optional newWorkspace:string and/or changeSubmitted:boolean
+		opts.newOwner = newOwner
+		this.editChange(roboWorkspace, changelist, opts)
+	}
+
 	where(roboWorkspace: RoboWorkspace, clientPath: string) {
 		return this._execP4Ztag(roboWorkspace, ['where', clientPath])
 	}
