Changes missed to go along with 24085573
(also mostly unnecessary import reorg imposed by addon)

diff --git a/Engine/Extras/RoboMerge/v3/src/robo/edgebot.ts b/Engine/Extras/RoboMerge/v3/src/robo/edgebot.ts
--- a/Engine/Extras/RoboMerge/v3/src/robo/edgebot.ts
+++ b/Engine/Extras/RoboMerge/v3/src/robo/edgebot.ts
@@ -5,24 +5,22 @@ import * as p4util from '../common/p4util';
 
 import { ContextualLogger } from "../common/logger";
 import { Recipients } from "../common/mailer";
-import { Change, ClientSpec, ConflictedResolveNFile, EditOwnerOpts, IntegrationSource } from "../common/perforce";
-import { coercePerforceWorkspace, getRootDirectoryForBranch } from "../common/perforce";
-import { IntegrationTarget, isExecP4Error, OpenedFileRecord, PerforceContext, EXCLUSIVE_CHECKOUT_REGEX } from "../common/perforce";
+import { Change, ClientSpec, coercePerforceWorkspace, ConflictedResolveNFile, EditChangeOpts, EXCLUSIVE_CHECKOUT_REGEX, getRootDirectoryForBranch, IntegrationSource, IntegrationTarget, isExecP4Error, OpenedFileRecord, PerforceContext } from "../common/perforce";
 import { VersionReader } from "../common/version";
 import { EdgeBotInterface, IPCControls, ReconsiderArgs } from "./bot-interfaces";
-import { PauseState } from "./state-interfaces";
-import { BlockagePauseInfo, BlockagePauseInfoMinimal, EdgeStatusFields } from "./status-types";
 import { AlreadyIntegrated, Branch, ChangeInfo, ConflictingFile, Failure, MergeAction, PendingChange } from "./branch-interfaces";
+import { EdgeOptions } from "./branchdefs";
 import { PersistentConflict } from "./conflict-interfaces";
 import { BotEventTriggers } from "./events";
+import { Gate } from "./gate";
 import { NodeBot } from "./nodebot";
 import { isUserAKnownBot, postMessageToChannel } from "./notifications";
 import { PerforceStatefulBot } from "./perforce-stateful-bot";
 import { Context } from "./settings";
 import { SlackMessageStyles } from "./slack";
-import { EdgeOptions } from "./branchdefs";
+import { PauseState } from "./state-interfaces";
+import { BlockagePauseInfo, BlockagePauseInfoMinimal, EdgeStatusFields } from "./status-types";
 import { getIntegrationOwner } from "./targets";
-import { Gate } from "./gate";
 
 function matchPrefix(a: string, b: string) {
 	const len = Math.min(a.length, b.length)
@@ -882,7 +880,7 @@ class EdgeBotImpl extends PerforceStatefulBot {
 		}
 
 		// log if we couldn't find a workspace
-		const opts: EditOwnerOpts = {edgeServerAddress}
+		const opts: EditChangeOpts = {edgeServerAddress}
 
 		if (targetWorkspace) {
 			this.edgeBotLogger.info(`Moving CL ${changenum} to workspace '${targetWorkspace}'`)
