Approval improvements
* Approval workflow can be set to not cause robomerge blockage
* It is no longer necessary that the author have a workspace mapped to the target stream for the approval changelist to be submitted
* Reconsidered changes will now have the reviews created/be submitted on behalf of the original author, not the user requesting the reconsider
* Approval channel is now linked in direct messages and robomerge channel messsages
* Approval messages now specify which stream the approval is for

diff --git a/Engine/Extras/RoboMerge/v3/functional_tests/framework/src/framework.ts b/Engine/Extras/RoboMerge/v3/functional_tests/framework/src/framework.ts
--- a/Engine/Extras/RoboMerge/v3/functional_tests/framework/src/framework.ts
+++ b/Engine/Extras/RoboMerge/v3/functional_tests/framework/src/framework.ts
@@ -206,9 +206,9 @@ type EdgeOptionFields = {
 	ignoreInCycleDetection: boolean
 
 	approval: {
-		description: string,
-		channelName: string
+		description: string
 		channelId: string
+		block: boolean
 	}
 }
 
diff --git a/Engine/Extras/RoboMerge/v3/functional_tests/framework/src/tests/approval.ts b/Engine/Extras/RoboMerge/v3/functional_tests/framework/src/tests/approval.ts
--- a/Engine/Extras/RoboMerge/v3/functional_tests/framework/src/tests/approval.ts
+++ b/Engine/Extras/RoboMerge/v3/functional_tests/framework/src/tests/approval.ts
@@ -31,8 +31,8 @@ export class Approval extends SimpleMainAndReleaseTestBase {
 		  { from: this.fullBranchName('Release'), to: this.fullBranchName('Main')
 		  , approval: {
 		      description: 'hey there!',
-		      channelName: 'approvers',
-		      channelId: 'abc123'
+		      channelId: 'abc123',
+		      block: true
 		    }
 		  }
 		]
diff --git a/Engine/Extras/RoboMerge/v3/src/robo/branchdefs.ts b/Engine/Extras/RoboMerge/v3/src/robo/branchdefs.ts
--- a/Engine/Extras/RoboMerge/v3/src/robo/branchdefs.ts
+++ b/Engine/Extras/RoboMerge/v3/src/robo/branchdefs.ts
@@ -162,8 +162,8 @@ const edgeOptionFieldsPrototype = {
 
 	approval: {
 		description: '',
-		channelName: '',
-		channelId: ''
+		channelId: '',
+		block: true
 	}
 }
 
@@ -461,7 +461,7 @@ export class BranchDefs {
 
 				if (edge.approval) {
 					// should be replaced by generic handling of objects in the prototype
-					if (!edge.approval.description || !edge.approval.channelName || !edge.approval.channelId) {
+					if (!edge.approval.description || !edge.approval.channelId) {
 						throw new Error(`Invalid approval settings for edge ${edge.from}->${edge.to}`)
 					}
 				}
diff --git a/Engine/Extras/RoboMerge/v3/src/robo/edgebot.ts b/Engine/Extras/RoboMerge/v3/src/robo/edgebot.ts
--- a/Engine/Extras/RoboMerge/v3/src/robo/edgebot.ts
+++ b/Engine/Extras/RoboMerge/v3/src/robo/edgebot.ts
@@ -308,7 +308,7 @@ class EdgeBotImpl extends PerforceStatefulBot {
 
 		if (pending.change.userRequest) {
 			const shelfMsg = `${owner}, please merge this change by hand.\nMore info at ${this.sourceNode.getBotUrl()}\n\n` + failure.description
-			await this.shelveChangelist(pending, shelfMsg)
+			await this.shelveChangelist(pending, false, shelfMsg)
 			this.edgeBotLogger.info(`${logMessage}. Shelved CL ${pending.newCl} for ${owner} to resolve manually (from reconsider).`)
 			return
 		}
@@ -658,12 +658,12 @@ class EdgeBotImpl extends PerforceStatefulBot {
 				const approval = this.options.approval
 
 				// shelve
-				await this.shelveChangelist(pending)
+				await this.shelveChangelist(pending, true)
 	
 				// still notify as a 'failure' to trigger normal blockage mechanism
 				failure = { kind: 'Approval required', description: approval.description }
 	
-				if (pending.change.userRequest) {
+				if (pending.change.userRequest || approval.block === false) {
 					this.sourceNode.reportApprovalRequired(approval, pending) 
 				}
 				else {
@@ -784,7 +784,7 @@ class EdgeBotImpl extends PerforceStatefulBot {
 		return new EdgeIntegrationDetails('ok', finalCl)
 	}
 
-	private async shelveChangelist(pending: PendingChange, reason?: string) {
+	private async shelveChangelist(pending: PendingChange, forApproval?: boolean, reason?: string) {
 		const changenum = pending.newCl
 		const owner = getIntegrationOwner(pending) || pending.change.author
 		this._log_action(`Shelving CL ${changenum} (change owned by ${owner})`)
@@ -804,6 +804,10 @@ class EdgeBotImpl extends PerforceStatefulBot {
 			final_desc += pending.change.additionalDescriptionText
 		}
 
+		/*if (forApproval) {
+			final_desc += "\n#review"
+		}*/
+
 		const edgeServerAddress = pending.change.edgeServerToHostShelf && pending.change.edgeServerToHostShelf.address
 		const destRoboWorkspace = pending.change.targetWorkspaceOverride || this.targetBranch.workspace
 
@@ -843,13 +847,13 @@ class EdgeBotImpl extends PerforceStatefulBot {
 
 		// figure out what workspace to put it in
 		const branch_stream = this.targetBranch.stream ? this.targetBranch.stream.toLowerCase() : null
-		let targetWorkspace: string | null = null
+		let targetWorkspace: string | undefined = undefined
 		// Check for specified workspace from createShelf operation
 		if (pending.change.targetWorkspaceForShelf) {
 			targetWorkspace = pending.change.targetWorkspaceForShelf
 		}
 		// Find a suitable workspace from one of the owner's workspaces
-		else {
+		else if (!forApproval) {
 			// use p4.find_workspaces to find a workspace (owned by the user) for this change if this is a stream branch
 			// not supporting edge servers yet - really need to send a separate shelf request instead
 			// so specify to only look for workspaces on the commit server
@@ -886,12 +890,13 @@ class EdgeBotImpl extends PerforceStatefulBot {
 			this.edgeBotLogger.info(`Moving CL ${changenum} to workspace '${targetWorkspace}'`)
 			opts.newWorkspace = targetWorkspace
 		}
-		else {
+		else if (!forApproval) {
 			this.edgeBotLogger.warn(`Unable to find appropriate workspace for ${owner}` + (branch_stream || this.targetBranch.name))
 		}
 
 		// edit the owner to the author so they can resolve and submit themselves
-		await this.p4.editOwner(destRoboWorkspace, changenum, owner, opts)
+		let shelfOwner = forApproval ? pending.change.author : owner
+		await this.p4.editOwner(destRoboWorkspace, changenum, shelfOwner, opts)
 	}
 
 	onGlobalChange(info: ChangeInfo) {
diff --git a/Engine/Extras/RoboMerge/v3/src/robo/nodebot.ts b/Engine/Extras/RoboMerge/v3/src/robo/nodebot.ts
--- a/Engine/Extras/RoboMerge/v3/src/robo/nodebot.ts
+++ b/Engine/Extras/RoboMerge/v3/src/robo/nodebot.ts
@@ -2008,7 +2008,8 @@ export class NodeBot extends PerforceStatefulBot implements NodeBotInterface {
 
 			const opts: SlackMessage = { 
 				title:'', 
-				text: `${channelPing}'s change needs to be approved.`, 
+				text: `${channelPing}'s change in ${pending.action.branch.name} needs to be approved.\n\n` +
+						approval.description, 
 				style: SlackMessageStyles.DANGER, 
 				fields,
 				mrkdwn: true,
diff --git a/Engine/Extras/RoboMerge/v3/src/robo/notifications.ts b/Engine/Extras/RoboMerge/v3/src/robo/notifications.ts
--- a/Engine/Extras/RoboMerge/v3/src/robo/notifications.ts
+++ b/Engine/Extras/RoboMerge/v3/src/robo/notifications.ts
@@ -490,7 +490,7 @@ export class BotNotifications implements BotEventHandler {
 
 		const isBotUser = isUserAKnownBot(blockage.owner)
 		const text =
-			blockage.approval ?				`${channelPing}'s change needs to be approved in ${blockage.approval.settings.channelName}` :
+			blockage.approval ?				`${channelPing}'s change needs to be approved in <#${blockage.approval.settings.channelId}>` :
 			isBotUser ? 									`Blockage caused by \`${blockage.owner}\` commit!` :
 			blockage.failure.kind === 'Too many files' ?	`${channelPing}, please request a shelf for this large changelist` :
 															`${channelPing}, please resolve the following ${issue}:`
@@ -596,10 +596,10 @@ export class BotNotifications implements BotEventHandler {
 			if (blockage.approval) {
 				dm = {
 					title: 'Approval needed to commit to ' + targetBranch.name,
-					text: `Your change has been shelved in ${makeClLink(cl)} and sent to ${blockage.approval.settings.channelName} for approval\n\n` +
+					text: `Your change has been shelved in ${makeClLink(cl)} and sent to <#${blockage.approval.settings.channelId}> for approval\n\n` +
 							blockage.approval.settings.description,
 					channel: "",
-					mrkdwn: false
+					mrkdwn: true
 				}
 			}
 			else {
