Fix race condition

diff --git a/Engine/Extras/RoboMerge/v3/src/common/perforce.ts b/Engine/Extras/RoboMerge/v3/src/common/perforce.ts
--- a/Engine/Extras/RoboMerge/v3/src/common/perforce.ts
+++ b/Engine/Extras/RoboMerge/v3/src/common/perforce.ts
@@ -1282,7 +1282,7 @@ export class PerforceContext {
 	// output format is just error or not
 	async editDescription(roboWorkspace: RoboWorkspace, changelist: number, description: string, edgeServerAddress?: string) {
 		const opts: EditChangeOpts = {newDescription: description, edgeServerAddress}
-		this.editChange(roboWorkspace, changelist, opts)
+		await this.editChange(roboWorkspace, changelist, opts)
 	}
 
 	// change the owner of an existing CL
@@ -1290,7 +1290,7 @@ export class PerforceContext {
 	async editOwner(roboWorkspace: RoboWorkspace, changelist: number, newOwner: string, opts?: EditChangeOpts) {
 		opts = opts || {}; // optional newWorkspace:string and/or changeSubmitted:boolean
 		opts.newOwner = newOwner
-		this.editChange(roboWorkspace, changelist, opts)
+		await this.editChange(roboWorkspace, changelist, opts)
 	}
 
 	where(roboWorkspace: RoboWorkspace, clientPath: string) {
