Adding a bot group for twinmotion contractors

#rnx

diff --git a/Engine/Extras/RoboMerge/v3/config/ldap.cfg.json b/Engine/Extras/RoboMerge/v3/config/ldap.cfg.json
--- a/Engine/Extras/RoboMerge/v3/config/ldap.cfg.json
+++ b/Engine/Extras/RoboMerge/v3/config/ldap.cfg.json
@@ -18,6 +18,10 @@
 		{
 			"group": "Robomerge_FTE_Access",
 			"tags": ["fte"]
+		},
+		{
+			"group": "Robomerge_Contractor_Twinmotion_Access",
+			"tags": ["twinmotion"]
 		}
 	]
 }
\ No newline at end of file
