[NO NON-PII COMMIT MESSAGE]

diff --git a/Engine/Extras/RoboMerge/v3/package.json b/Engine/Extras/RoboMerge/v3/package.json
--- a/Engine/Extras/RoboMerge/v3/package.json
+++ b/Engine/Extras/RoboMerge/v3/package.json
@@ -10,8 +10,7 @@
   "license": "UNLICENSED",
   "scripts": {
     "start": "node dist/robo/watchdog.js -botname=TEST -devMode",
-    "test": "node dist/robo/unit-tests.js",
-    "postinstall": "node install-graphviz.js"
+    "test": "node dist/robo/unit-tests.js"
   },
   "dependencies": {
     "@hpcc-js/graph": "^2.27.0",
diff --git a/Engine/Extras/RoboMerge/v3/src/robo/conflicts.ts b/Engine/Extras/RoboMerge/v3/src/robo/conflicts.ts
--- a/Engine/Extras/RoboMerge/v3/src/robo/conflicts.ts
+++ b/Engine/Extras/RoboMerge/v3/src/robo/conflicts.ts
@@ -233,7 +233,7 @@ export class Conflicts {
 	 * This happens before a tick, so that all other bots will have had a chance
 	 * to update, so that resolving changelists will have been seen
 	 */
-	checkForResolvedConflicts(nodeCL: Number, edges: Map<string, EdgeBot>) {
+	checkForResolvedConflicts(nodeCL: number, edges: Map<string, EdgeBot>) {
 		if (this.conflicts.length === 0) {
 			return
 		}
