Ensure that an existing client workspace maps the correct stream before using it. Otherwise rewrite it to point to the desired stream.

diff --git a/Engine/Extras/RoboMerge/v3/src/robo/robo.ts b/Engine/Extras/RoboMerge/v3/src/robo/robo.ts
--- a/Engine/Extras/RoboMerge/v3/src/robo/robo.ts
+++ b/Engine/Extras/RoboMerge/v3/src/robo/robo.ts
@@ -269,7 +269,7 @@ async function _initWorkspacesForGraphBot(graphBot: GraphBot, existingWorkspaces
 
 		// see if we already have this workspace
 		const clientSpec = existingWorkspaces.get(ws)
-		if (clientSpec) {
+		if (clientSpec && clientSpec.Stream === branch.stream) {
 			if (clientSpec.IsUnloaded) {
 				reloadedWorkspaces.push(robo.p4.reloadWorkspace(ws))
 			}
