Send failures to lookup email address when trying to invite user to alerts

diff --git a/Engine/Extras/RoboMerge/v3/src/robo/notifications.ts b/Engine/Extras/RoboMerge/v3/src/robo/notifications.ts
--- a/Engine/Extras/RoboMerge/v3/src/robo/notifications.ts
+++ b/Engine/Extras/RoboMerge/v3/src/robo/notifications.ts
@@ -372,7 +372,9 @@ export class SlackMessages {
 						/* this is expected and fine */ 
 					}
 					else {
-						this.smLogger.error(`Error inviting ${emailAddress} (${user}) to channel ${channel}: ${result.error}\n${result.failed_user_ids}`)
+						const errorMsg = `Error inviting ${emailAddress} (${user}) to channel ${channel}: ${result.error}\n${result.failed_user_ids}`
+						this.smLogger.error(errorMsg)
+						postToRobomergeAlerts(errorMsg)
 					}
 				}
 				else if (result.error === "user_is_restricted") {
@@ -386,12 +388,16 @@ export class SlackMessages {
 					}
 				}
 				else {
-					this.smLogger.error(`Error inviting ${emailAddress} (${user}) to channel ${channel}: ${result.error}`)
+					const errorMsg = `Error inviting ${emailAddress} (${user}) to channel ${channel}: ${result.error}`
+					this.smLogger.error(errorMsg)
+					postToRobomergeAlerts(errorMsg)
 				}
 			}
 		}
 		else {
-			this.smLogger.error(`Unable to add ${emailAddress} to channel ${channel}`)
+			const errorMsg = `Unable to add ${emailAddress} to channel ${channel}`
+			this.smLogger.error(errorMsg)
+			postToRobomergeAlerts(errorMsg)
 		}
 	}
 
