Don't overwrite open branchspec files when syncing in devMode
Reload branchspecs when they are saved in devMode

diff --git a/Engine/Extras/RoboMerge/v3/src/robo/autobranchupdater.ts b/Engine/Extras/RoboMerge/v3/src/robo/autobranchupdater.ts
--- a/Engine/Extras/RoboMerge/v3/src/robo/autobranchupdater.ts
+++ b/Engine/Extras/RoboMerge/v3/src/robo/autobranchupdater.ts
@@ -14,6 +14,7 @@ const DISABLE = false
 type AutoBranchUpdaterConfig = {
 	rootPath: string		//				process.env('ROBO_BRANCHSPECS_ROOT_PATH')
 	workspace: Workspace	// { directory:	process.env('ROBO_BRANCHSPECS_DIRECTORY')
+	devMode: boolean
 }							// , name:		process.env('ROBO_BRANCHSPECS_WORKSPACE') }
 
 type MirrorPaths = {
@@ -35,6 +36,7 @@ export class AutoBranchUpdater implements Bot {
 	private readonly abuLogger: ContextualLogger
 	
 	lastCl: number
+	lastModifiedTime: number // Only used in devmode
 
 	// public: used by NodeBot
 	isRunning = false
@@ -82,25 +84,19 @@ export class AutoBranchUpdater implements Bot {
 			fs.mkdirSync(workspaceDir);
 		}
 
-		await this.p4.sync(config.workspace, `${bsRoot}@${change.change}`, [P4_FORCE])
+		const opts = !config.devMode ? [P4_FORCE] : undefined 
+		await this.p4.sync(config.workspace, `${bsRoot}@${change.change}`, opts)
 	}
 	
 	async start() {
-		this.isRunning = true;
-		this.abuLogger.info(`Began monitoring ${this.graphBot.branchGraph.botname} branch specs at CL ${this.lastCl}`);
-
-// this got lost in the mirror refactor
 
-		// if (!DISABLE && this.mirror) {
-		// 	const workspace = await this.p4.find_workspace_by_name(this.mirror.workspace.name)
-		// 	if (workspace.length === 0) {
-		// 		await this.p4.newWorkspace(this.mirror.workspace.name, {
-		// 			Stream: this.mirror.stream,
-		// 			Root: AutoBranchUpdater.config!.workspace.directory
-		// 		})
-		// 	}
-		// }
+		if (AutoBranchUpdater.config.devMode) {
+			const stats = fs.statSync(`${this.workspace.directory}/${this.graphBot.filename}`);
+			this.lastModifiedTime = new Date(stats.mtime).getTime();
+		}
 
+		this.isRunning = true;
+		this.abuLogger.info(`Began monitoring ${this.graphBot.branchGraph.botname} branch specs at CL ${this.lastCl}`);
 	}
 
 	async tick() {
@@ -108,22 +104,33 @@ export class AutoBranchUpdater implements Bot {
 			return false
 		}
 
-		let change;
+		let reloadBranchDefs = false;
 		try {
-			change = await this.p4.latestChange(this.filePath);
+			const change = await this.p4.latestChange(this.filePath);
+			if (change !== null && change.change > this.lastCl) {
+				await this.p4.sync(this.workspace, `${this.filePath}@${change.change}`);
+	
+				// set this to be the last changelist regardless of success - if it failed due to a broken
+				// .json file, the file will have to be recommitted anyway
+				this.lastCl = change.change;
+				reloadBranchDefs = true;
+			}
 		}
 		catch (err) {
 			this.abuLogger.printException(err, 'Branch specs: error while querying P4 for changes');
 			return false
 		}
 
-		if (change !== null && change.change > this.lastCl) {
-			await this.p4.sync(this.workspace, `${this.filePath}@${change.change}`);
-
-			// set this to be the last changelist regardless of success - if it failed due to a broken
-			// .json file, the file will have to be recommitted anyway
-			this.lastCl = change.change;
+		if (AutoBranchUpdater.config.devMode) {
+			const stats = fs.statSync(`${this.workspace.directory}/${this.graphBot.filename}`);
+			const newModifiedTime = new Date(stats.mtime).getTime();
+			if (this.lastModifiedTime !== newModifiedTime) {
+				this.lastModifiedTime = newModifiedTime
+				reloadBranchDefs = true
+			}
+		}
 
+		if (reloadBranchDefs) {
 			await this._tryReloadBranchDefs();
 		}
 		return true
diff --git a/Engine/Extras/RoboMerge/v3/src/robo/robo.ts b/Engine/Extras/RoboMerge/v3/src/robo/robo.ts
--- a/Engine/Extras/RoboMerge/v3/src/robo/robo.ts
+++ b/Engine/Extras/RoboMerge/v3/src/robo/robo.ts
@@ -387,7 +387,8 @@ async function init(logger: ContextualLogger) {
 		const branchSpecsAbsPath = fs.realpathSync(args.branchSpecsDirectory)
 		const autoUpdaterConfig = {
 			rootPath: args.branchSpecsRootPath,
-			workspace: {directory: branchSpecsAbsPath, name: args.branchSpecsWorkspace}
+			workspace: {directory: branchSpecsAbsPath, name: args.branchSpecsWorkspace},
+			devMode: args.devMode
 		}
 
 		// Ensure we have a workspace for branch specs
