Use #robomerge-author instead of waiting on robomerge to finish changing owner
Extend waiting time to 2 minutes for swarm submits

diff --git a/Engine/Extras/RoboMerge/v3/src/robo/nodebot.ts b/Engine/Extras/RoboMerge/v3/src/robo/nodebot.ts
--- a/Engine/Extras/RoboMerge/v3/src/robo/nodebot.ts
+++ b/Engine/Extras/RoboMerge/v3/src/robo/nodebot.ts
@@ -1360,23 +1360,31 @@ export class NodeBot extends PerforceStatefulBot implements NodeBotInterface {
 		for (let changeIndex = 0; changeIndex < changes.length; ++changeIndex) {
 			const change = changes[changeIndex]
 
-			// If the user of a change is robomerge or this is a change submitted via swarm on behalf of someone
-			// don't process the change for a minute to give the change owner command time to be processed
-			if ((Date.now() - ((change.time || 0) * 1000)) < 60000) {
-				if (change.user === "robomerge") {
-					this.nodeBotLogger.info(`Delaying processing of ${change.change} while robomerge changes owner.  Now: ${Date.now()} Change time: ${change.time}`)
-					return
+			// If the user of the CL is robomerge, lookup the actual author from the tag
+			if (change.user === "robomerge") {
+				let bAuthorFound = false
+				for (const line of change.desc.split('\n')) {
+					const authorMatch = line.match("#ROBOMERGE-AUTHOR: (.*)")
+					if (authorMatch) {
+						change.user = authorMatch[1]
+						bAuthorFound = true
+						break
+					}
+				}
+				if (!bAuthorFound) {
+					this.nodeBotLogger.warn(`Processing change ${change.change} by user robomerge. Now: ${Date.now()} Change time: ${change.time}`)
 				}
-				else if (change.client.startsWith("swarm-")) {
+			}
+			// If this is a change submitted via swarm on behalf of someone don't process the
+			// change for a few minutes to give the change owner command time to be processed
+			else if (change.client.startsWith("swarm-")) {
+				// TODO: Use the Swarm API to lookup the actual author
+				if ((Date.now() - ((change.time || 0) * 1000)) < 120000) {
 					this.nodeBotLogger.info(`Delaying processing of ${change.change} to give swarm time to change owner.  Now: ${Date.now()} Change time: ${change.time}`)
 					return
 				}
 			}
 
-			if (change.user === "robomerge") {
-				this.nodeBotLogger.warn(`Processing change ${change.change} by user robomerge. Now: ${Date.now()} Change time: ${change.time}`)
-            }
-
 			const changeResult = await this._processAndMergeCl(availableEdges, change, false)
 
 			// Exit immediately on syntax errors
