Fix p4 resolve ztag parsing resulting in incorrect results in blockage messages

diff --git a/Engine/Extras/RoboMerge/v3/src/common/perforce.ts b/Engine/Extras/RoboMerge/v3/src/common/perforce.ts
--- a/Engine/Extras/RoboMerge/v3/src/common/perforce.ts
+++ b/Engine/Extras/RoboMerge/v3/src/common/perforce.ts
@@ -31,6 +31,7 @@ const ztag_group_rex = /\n\n\.\.\.\s/;
 const ztag_field_rex = /(?:\n|^)\.\.\.\s/;
 const newline_rex = /\r\n|\n|\r/g;
 const integer_rex = /^[1-9][0-9]*\s*$/;
+const resolveGarbage_rex = /^(Branch resolve:|at: branch|ay: ignore)$/gm;
 
 export interface BranchSpec {
 	name: string;
@@ -98,7 +99,7 @@ function parseTrace(response: string): [string, string] {
 
 // parse the perforce tagged output format into an array of objects
 // TODO: probably should switch this to scrape Python dictionary format (-G) since ztag is super inconsistent with multiline fields
-export function parseZTag(buffer: string, multiLine?: boolean) {
+export function parseZTag(buffer: string, opts?: ExecZtagOpts) {
 	let output = [];
 
 	// check for error lines ahead of the first ztag field
@@ -116,6 +117,12 @@ export function parseZTag(buffer: string, multiLine?: boolean) {
 		buffer = "";
 	}
 
+	// resolve ztag can have some garbage that causes issues with
+	// parsing the groups, so we're just going to strip them out
+	if (opts && opts.resolve) {
+		buffer = buffer.replaceAll(resolveGarbage_rex,"")
+	}
+
 	// split into groups
 	let groups = buffer.split(ztag_group_rex);
 	for (let i = 0; i < groups.length; ++i) {
@@ -141,7 +148,7 @@ export function parseZTag(buffer: string, multiLine?: boolean) {
 			if (s >= 0) {
 				key = pair.substr(0, s);
 				value = pair.substr(s + 1);
-				if (value.indexOf('\n') >= 0 && !multiLine) {
+				if (value.indexOf('\n') >= 0 && !(opts && opts.multiline)) {
 					let lines = value.split('\n');
 					value = lines.shift();
 					text = text.concat(lines.filter((str) => { return str !== ""; }));
@@ -278,6 +285,7 @@ interface ExecOpts {
 
 interface ExecZtagOpts extends ExecOpts {
 	multiline?: boolean;
+	resolve?: boolean; // hacky solution to clear certain problematic lines out of resolve ztags
 }
 
 export interface EditChangeOpts {
@@ -967,8 +975,7 @@ export class PerforceContext {
 
 		let dashNresult: string[]
 		try {
-			//dashNresult = await this._execP4(workspace, ['resolve', '-N', `-c${changelist}`])
-			dashNresult = await this._execP4Ztag(workspace, ['resolve', '-N', `-c${changelist}`], {edgeServerAddress})
+			dashNresult = await this._execP4Ztag(workspace, ['resolve', '-N', `-c${changelist}`], {edgeServerAddress, resolve: true})
 		}
 		catch (reason) {
 			if (!isExecP4Error(reason)) {
@@ -1512,7 +1519,7 @@ export class PerforceContext {
 
 	static async _execP4Ztag(logger: ContextualLogger, roboWorkspace: RoboWorkspace, args: string[], opts?: ExecZtagOpts) {
 		const workspace = coercePerforceWorkspace(roboWorkspace);
-		return parseZTag(await PerforceContext._execP4(logger, workspace, ['-ztag', ...args], opts), opts && opts.multiline);
+		return parseZTag(await PerforceContext._execP4(logger, workspace, ['-ztag', ...args], opts), opts);
 	}
 
 	private async _execP4Ztag(roboWorkspace: RoboWorkspace, args: string[], opts?: ExecZtagOpts) {
