[NO NON-PII COMMIT MESSAGE]

diff --git a/Engine/Extras/RoboMerge/v3/README.md b/Engine/Extras/RoboMerge/v3/README.md
--- a/Engine/Extras/RoboMerge/v3/README.md
+++ b/Engine/Extras/RoboMerge/v3/README.md
@@ -142,3 +142,28 @@ Dis-allow integration between 5am and 9pm UTC
 ```
 { "integrationWindow": [{ "startHourUTC": 5, "durationHours": 16 }], "invertIntegrationWindow": true }
 ```
+
+### Upgrading to newer Helix/NodeJS/Typescript libraries
+
+#### Helix
+
+The last line of `helix.Dockerfile` contains the version number.
+
+#### NodeJS/Typescript
+
+The NodeJS/Typescript versions are specified in `helix-and-node.Dockerfile`:
+
+e.g. Node16
+> `RUN curl -sL https://rpm.nodesource.com/setup_16.x | bash -`
+
+Typescript 5.0
+> `RUN npm install -g typescript@^5.0`
+
+If either of these docker files and modified, a new version of the image should be built and pushed to the Docker Hub. In the Epic internal case, for example:
+
+```
+docker build -f helix.Dockerfile -t helix
+docker tag helix hub.epicgames.ol.net/epicgames/helix:2021.1
+docker push hub.epicgames.ol.net/epicgames/helix:2021.1
+```
+
