Put the channel name, not just the ID in the alert message

diff --git a/Engine/Extras/RoboMerge/v3/src/robo/notifications.ts b/Engine/Extras/RoboMerge/v3/src/robo/notifications.ts
--- a/Engine/Extras/RoboMerge/v3/src/robo/notifications.ts
+++ b/Engine/Extras/RoboMerge/v3/src/robo/notifications.ts
@@ -372,7 +372,7 @@ export class SlackMessages {
 						/* this is expected and fine */ 
 					}
 					else {
-						const errorMsg = `Error inviting ${emailAddress} (${user}) to channel ${channel}: ${result.error}\n${result.failed_user_ids}`
+						const errorMsg = `Error inviting ${emailAddress} (${user}) to channel <#${channel}>: ${result.error}\n${result.failed_user_ids}`
 						this.smLogger.error(errorMsg)
 						postToRobomergeAlerts(errorMsg)
 					}
@@ -388,14 +388,14 @@ export class SlackMessages {
 					}
 				}
 				else {
-					const errorMsg = `Error inviting ${emailAddress} (${user}) to channel ${channel}: ${result.error}`
+					const errorMsg = `Error inviting ${emailAddress} (${user}) to channel <#${channel}>: ${result.error}`
 					this.smLogger.error(errorMsg)
 					postToRobomergeAlerts(errorMsg)
 				}
 			}
 		}
 		else {
-			const errorMsg = `Unable to add ${emailAddress} to channel ${channel}`
+			const errorMsg = `Unable to add ${emailAddress} to channel <#${channel}>`
 			this.smLogger.error(errorMsg)
 			postToRobomergeAlerts(errorMsg)
 		}
