Don't spam exceptions in dev mode if there is a new branchspec file that hasn't been submitted to perforce

diff --git a/Engine/Extras/RoboMerge/v3/src/robo/autobranchupdater.ts b/Engine/Extras/RoboMerge/v3/src/robo/autobranchupdater.ts
--- a/Engine/Extras/RoboMerge/v3/src/robo/autobranchupdater.ts
+++ b/Engine/Extras/RoboMerge/v3/src/robo/autobranchupdater.ts
@@ -117,8 +117,11 @@ export class AutoBranchUpdater implements Bot {
 			}
 		}
 		catch (err) {
-			this.abuLogger.printException(err, 'Branch specs: error while querying P4 for changes');
-			return false
+			// if we're in devmode we support files open for add or not in perforce so absord the exception
+			if (!AutoBranchUpdater.config.devMode) {
+				this.abuLogger.printException(err, 'Branch specs: error while querying P4 for changes');
+				return false
+			}
 		}
 
 		if (AutoBranchUpdater.config.devMode) {
