Fix crash when reloading a branchspec that has not been submitted to the depot (in devmode)

diff --git a/Engine/Extras/RoboMerge/v3/src/common/p4util.ts b/Engine/Extras/RoboMerge/v3/src/common/p4util.ts
--- a/Engine/Extras/RoboMerge/v3/src/common/p4util.ts
+++ b/Engine/Extras/RoboMerge/v3/src/common/p4util.ts
@@ -103,7 +103,7 @@ export async function convertIntegrateToEdit(p4: PerforceContext, roboWorkspace:
 	if (notSyncedOpenFiles.length !== 0) {
 // don't warn - this will almost always be the case for RoboMerge
 //		_warnNotSyncedOpenFiles(notSyncedOpenFiles);
-		await Promise.all(notSyncedOpenFiles.map(file => p4.sync(workspace, `${file.clientFile}#${file.rev}`, ['-k'])))
+		await Promise.all(notSyncedOpenFiles.map(file => p4.sync(workspace, `${file.clientFile}#${file.rev}`, {opts: ['-k']})))
 	}
 
 	// Perform basic actions
@@ -220,7 +220,7 @@ export async function cleanWorkspaces(p4: PerforceContext, workspaces: [string,
 	}
 
 	p4utilsLogger.info('Resetting all workspaces to revision 0')
-	await Promise.all(workspaces.map(([name, root]) => p4.sync(name, root + '#0', undefined, edgeServerAddress)))
+	await Promise.all(workspaces.map(([name, root]) => p4.sync(name, root + '#0', {edgeServerAddress})))
 }
 
 export async function getWorkspacesForUser(p4: PerforceContext, user: string, edgeServerAddress?: string) {
diff --git a/Engine/Extras/RoboMerge/v3/src/common/perforce.ts b/Engine/Extras/RoboMerge/v3/src/common/perforce.ts
--- a/Engine/Extras/RoboMerge/v3/src/common/perforce.ts
+++ b/Engine/Extras/RoboMerge/v3/src/common/perforce.ts
@@ -296,6 +296,12 @@ export interface EditChangeOpts {
 	edgeServerAddress?: string
 }
 
+export interface SyncParams {
+	opts?: string[]
+	edgeServerAddress?: string
+	okToFail?: boolean // default is false
+}
+
 export interface ConflictedResolveNFile {
 	clientFile: string // Workspace path on local disk
 	targetDepotFile?: string // Target Depot path in P4 Depot
@@ -657,12 +663,12 @@ export class PerforceContext {
 	}
 
 	// sync the depot path specified
-	async sync(roboWorkspace: RoboWorkspace, depotPath: string, opts?: string[], edgeServerAddress?: string) {
+	async sync(roboWorkspace: RoboWorkspace, depotPath: string, params?: SyncParams) {
 		const workspace = coercePerforceWorkspace(roboWorkspace);
-		let args = edgeServerAddress ? ['-p', edgeServerAddress] : []
+		let args = params && params.edgeServerAddress ? ['-p', params.edgeServerAddress] : []
 		args.push('sync')
-		if (opts) {
-			args = [...args, ...opts]
+		if (params && params.opts) {
+			args = [...args, ...params.opts]
 		}
 		args.push(depotPath)
 		try {
@@ -672,10 +678,13 @@ export class PerforceContext {
 			if (!isExecP4Error(reason)) {
 				throw reason
 			}
-			const [err, output] = reason
-			// this is an acceptable non-error case for us
-			if (!output || typeof output !== "string" || !output.trim().endsWith("up-to-date."))
-				throw err
+			if (!params || !params.okToFail) {
+				const [err, output] = reason
+				// this is an acceptable non-error case for us
+				if (!output || typeof output !== "string" || !output.trim().endsWith("up-to-date.")) {
+					throw err
+				}
+			}
 		}
 	}
 
@@ -685,7 +694,7 @@ export class PerforceContext {
 		if (!change) {
 			throw new Error('Unable to find changelist');
 		}
-		await this.sync(workspace, `${depotPath}@${change.change}`, opts);
+		await this.sync(workspace, `${depotPath}@${change.change}`, {opts});
 		return change.change;
 	}
 
diff --git a/Engine/Extras/RoboMerge/v3/src/robo/autobranchupdater.ts b/Engine/Extras/RoboMerge/v3/src/robo/autobranchupdater.ts
--- a/Engine/Extras/RoboMerge/v3/src/robo/autobranchupdater.ts
+++ b/Engine/Extras/RoboMerge/v3/src/robo/autobranchupdater.ts
@@ -85,7 +85,7 @@ export class AutoBranchUpdater implements Bot {
 		}
 
 		const opts = !config.devMode ? [P4_FORCE] : undefined 
-		await this.p4.sync(config.workspace, `${bsRoot}@${change.change}`, opts)
+		await this.p4.sync(config.workspace, `${bsRoot}@${change.change}`, {opts})
 	}
 	
 	async start() {
@@ -173,7 +173,7 @@ export class AutoBranchUpdater implements Bot {
 		// NOTE: not awaiting next tick. Waiters on this function carry on as soon as we return
 		// this doesn't wait until all the branch bots have stopped, but that's ok - we're creating a new set of branch bots
 		process.nextTick(async () => {
-			await this.p4.sync(this.workspace, this.filePath)
+			await this.p4.sync(this.workspace, this.filePath, {okToFail:AutoBranchUpdater.config.devMode})
 
 			this.abuLogger.info(`Branch spec change detected: reloading ${botname} from CL#${this.lastCl}`)
 
@@ -231,7 +231,7 @@ export class AutoBranchUpdater implements Bot {
 
 		const {depotpath, realFilepath, mirrorFilepath} = workspace
 
-		await this.p4.sync(workspace, depotpath, [P4_FORCE])
+		await this.p4.sync(workspace, depotpath, {opts:[P4_FORCE]})
 		const cl = await this.p4.new_cl(workspace, "Updating mirror file\n#jira none\n#robomerge ignore\n")
 		await this.p4.edit(workspace, cl, depotpath)
 		await new Promise((done, _) => fs.copyFile(realFilepath, mirrorFilepath, done))
