Avoid catastrophic data loss when service crashes while persistence files are being written

diff --git a/Engine/Extras/RoboMerge/v3/src/robo/settings.ts b/Engine/Extras/RoboMerge/v3/src/robo/settings.ts
--- a/Engine/Extras/RoboMerge/v3/src/robo/settings.ts
+++ b/Engine/Extras/RoboMerge/v3/src/robo/settings.ts
@@ -107,8 +107,19 @@ export class Context {
 	}
 }
 
+type FilePurpose = 
+	'PERSISTENCE' |
+	'BACKUP' |
+	'TEMP'
+
+const FilePurposeExtensionMap = {
+	'PERSISTENCE' : '.json',
+	'BACKUP' : '.bak',
+	'TEMP' : '.save'
+}
+
 function getPersistenceFilepath(botname: string) {
-	return args.persistenceDir + `/${botname.toLowerCase()}.settings.json`
+	return args.persistenceDir + `/${botname.toLowerCase()}.settings`
 }
 
 export class Settings {
@@ -126,12 +137,12 @@ export class Settings {
 		// see if we should enable saves
 		this.enableSave = !process.env["NOSAVE"];
 		if (this.enableSave)
-			this.settingsLogger.verbose(`Reading settings from ${this.filename}`);
+			this.settingsLogger.verbose(`Reading settings from ${this.getFilename('PERSISTENCE')}`);
 		else
 			this.settingsLogger.info("Saving config has been disabled by NOSAVE environment variable");
 
 		// load the object from disk
-		let filebits = readFileToString(this.filename);
+		let filebits = readFileToString(this.getFilename('PERSISTENCE'));
 		let objVersion : semver.SemVer
 		if (filebits)
 			this.object = jsonlint.parse(filebits);
@@ -408,10 +419,24 @@ export class Settings {
 		return new Context(this, name);
 	}
 
+	private getFilename(purpose: FilePurpose) {
+		return this.filename + FilePurposeExtensionMap[purpose]
+	}
+
 	_saveObject() {
 		if (this.enableSave) {
-			let filebits = JSON.stringify(this.object, null, '  ');
-			fs.writeFileSync(this.filename, filebits, "utf8");
+
+			const persistentFile = this.getFilename('PERSISTENCE')
+			const tempFile = this.getFilename('TEMP')
+
+			if (fs.existsSync(persistentFile)) {
+				fs.copyFileSync(persistentFile,this.getFilename('BACKUP'))
+			}
+
+			let filebits = JSON.stringify(this.object, null, '  ')
+			fs.writeFileSync(tempFile, filebits, "utf8")
+
+			fs.copyFileSync(tempFile,persistentFile)
 		}
 	}
 }
