Do virtual integrations whenever the resulting product is not going to be provided to the user as a shelf (shelf requests, reconsiders, use of #manual)
Do not use -Rb -Rd as they simply require unnecessary resolves

diff --git a/Engine/Extras/RoboMerge/v3/src/common/perforce.ts b/Engine/Extras/RoboMerge/v3/src/common/perforce.ts
--- a/Engine/Extras/RoboMerge/v3/src/common/perforce.ts
+++ b/Engine/Extras/RoboMerge/v3/src/common/perforce.ts
@@ -296,6 +296,11 @@ export interface EditChangeOpts {
 	edgeServerAddress?: string
 }
 
+export interface IntegrateOpts {
+	edgeServerAddress?: string
+	virtual?: boolean
+}
+
 export interface SyncParams {
 	opts?: string[]
 	edgeServerAddress?: string
@@ -806,20 +811,25 @@ export class PerforceContext {
 	// integrate a CL from source to destination, resolve, and place the results in a new CL
 	// output format is true if the integration resolved or false if the integration wasn't necessary (still considered a success)
 	// failure to resolve is treated as an error condition
-	async integrate(roboWorkspace: RoboWorkspace, source: IntegrationSource, dest_changelist: number, target: IntegrationTarget, edgeServerAddress?: string)
+	async integrate(roboWorkspace: RoboWorkspace, source: IntegrationSource, dest_changelist: number, target: IntegrationTarget, opts?: IntegrateOpts)
 		: Promise<[string, (Change | string)[]]> {
+
+		opts = opts || {};
 		const workspace = coercePerforceWorkspace(roboWorkspace);
+
 		// build a command
 		let cmdList = [
 			"integrate",
 			"-Ob",
 			"-Or",
-			"-Rd",
-			"-Rb",
 			"-c" + dest_changelist
 		];
 		const range = `@${source.changelist},${source.changelist}`
 
+		if (opts.virtual) {
+			cmdList.push('-v')
+		}
+
 		let noSuchFilesPossible = false // Helper variable for error catching
 		// Branchspec -- takes priority above other possibilities
 		if (source.branchspec) {
@@ -849,7 +859,7 @@ export class PerforceContext {
 		// execute the P4 command
 		let changes;
 		try {
-			changes = await this._execP4Ztag(workspace, cmdList, { numRetries: 0, edgeServerAddress });
+			changes = await this._execP4Ztag(workspace, cmdList, { numRetries: 0, edgeServerAddress: opts.edgeServerAddress });
 		}
 		catch (reason) {
 			if (!isExecP4Error(reason)) {
@@ -974,6 +984,10 @@ export class PerforceContext {
 
 		// If resolveHelper() returned a string, we do not need to return a dashNResult
 		if (typeof fileInfo === 'string') {
+			// If our error is just there are no files to resolve, that is a success
+			if (fileInfo.startsWith('No file(s) to resolve.')) {
+				return new ResolveResult(null, 'success', this.logger)
+			}
 			return new ResolveResult(null, fileInfo, this.logger)
 		}
 
@@ -993,7 +1007,7 @@ export class PerforceContext {
 
 			let [err, output] = reason
 			
-			// If our error is just there are no files to resolve, 
+			// If our error is just there are no files to resolve, that is a success
 			if (output.startsWith('No file(s) to resolve.')) {
 				dashNresult = ['success']
 			} else {
diff --git a/Engine/Extras/RoboMerge/v3/src/robo/edgebot.ts b/Engine/Extras/RoboMerge/v3/src/robo/edgebot.ts
--- a/Engine/Extras/RoboMerge/v3/src/robo/edgebot.ts
+++ b/Engine/Extras/RoboMerge/v3/src/robo/edgebot.ts
@@ -524,8 +524,12 @@ class EdgeBotImpl extends PerforceStatefulBot {
 			stream: this.targetBranch.stream
 		}
 
+		// We want to do a virtual merge unless the change is going to be (or could be
+		// in userRequest/reconsider case) handed off to the user
+		const doVirtualMerge = !info.userRequest && !info.forceCreateAShelf && !target.flags.has('manual') 
+
 		this.currentIntegrationStartTimestamp = Date.now()
-		const [mode, results] = await this.p4.integrate(info.targetWorkspaceOverride, source, changenum, integTarget, edgeServerAddress)
+		const [mode, results] = await this.p4.integrate(info.targetWorkspaceOverride, source, changenum, integTarget, {edgeServerAddress, virtual: doVirtualMerge})
 
 		const pending: PendingChange = {change: info, action: target, newCl: changenum}
 
