Don't try to add bot users to channels

diff --git a/Engine/Extras/RoboMerge/v3/src/robo/notifications.ts b/Engine/Extras/RoboMerge/v3/src/robo/notifications.ts
--- a/Engine/Extras/RoboMerge/v3/src/robo/notifications.ts
+++ b/Engine/Extras/RoboMerge/v3/src/robo/notifications.ts
@@ -26,6 +26,7 @@ const SUNSET_PERSISTED_NOTIFICATIONS_DAYS = 30
 const EPIC_TIME_OPTIONS: DateTimeFormatOptions = {timeZone: 'EST5EDT', timeZoneName: 'short'}
 
 const KNOWN_BOT_NAMES = ['buildmachine', 'robomerge'];
+const KNOWN_BOT_EMAILS = [[REDACTED_EMAIL]'];
 
 let SLACK_TOKENS: {[name: string]: string} = {}
 const SLACK_DEV_DUMMY_TOKEN = 'dev'
@@ -64,7 +65,7 @@ export function notificationsInit(inArgs: Args) {
 }
 
 export function isUserAKnownBot(user: string) {
-	return KNOWN_BOT_NAMES.indexOf(user) !== -1
+	return KNOWN_BOT_NAMES.indexOf(user) !== -1 || KNOWN_BOT_EMAILS.indexOf(user) !== -1
 }
 
 
@@ -362,42 +363,44 @@ export class SlackMessages {
 
 	async addUserToChannel(emailAddress: string, channel: string, externalUser? :boolean)
 	{
-		const user = await this.getSlackUser(emailAddress)
-		if (user) {
-			const result = await this.slack.addUserToChannel(user, channel, externalUser)
-			if (!result.ok) {
-				if (result.error === "already_in_channel") { /* this is expected and fine */ }
-				else if (result.error === "failed_for_some_users") { 
-					if (result.failed_user_ids[user] === "unable_to_add_user_to_public_channel") {
-						/* this is expected and fine */ 
+		if (!isUserAKnownBot(emailAddress)) {
+			const user = await this.getSlackUser(emailAddress)
+			if (user) {
+				const result = await this.slack.addUserToChannel(user, channel, externalUser)
+				if (!result.ok) {
+					if (result.error === "already_in_channel") { /* this is expected and fine */ }
+					else if (result.error === "failed_for_some_users") { 
+						if (result.failed_user_ids[user] === "unable_to_add_user_to_public_channel") {
+							/* this is expected and fine */ 
+						}
+						else {
+							const errorMsg = `Error inviting ${emailAddress} (${user}) to channel <#${channel}>: ${result.error}\n${result.failed_user_ids}`
+							this.smLogger.error(errorMsg)
+							postToRobomergeAlerts(errorMsg)
+						}
+					}
+					else if (result.error === "user_is_restricted") {
+						// This an external user so we need to use the admin invite and validate that the channel ends in -ext
+						const channelInfo = await this.slack.getChannelInfo(channel)
+						if (channelInfo.channel.name.endsWith('-ext')) {
+							await this.addUserToChannel(emailAddress, channel, true)
+						}
+						else {
+							this.smLogger.error(`Cannot invite external user ${emailAddress} (${user}) to restricted channel ${channel}`)
+						}
 					}
 					else {
-						const errorMsg = `Error inviting ${emailAddress} (${user}) to channel <#${channel}>: ${result.error}\n${result.failed_user_ids}`
+						const errorMsg = `Error inviting ${emailAddress} (${user}) to channel <#${channel}>: ${result.error}`
 						this.smLogger.error(errorMsg)
 						postToRobomergeAlerts(errorMsg)
 					}
 				}
-				else if (result.error === "user_is_restricted") {
-					// This an external user so we need to use the admin invite and validate that the channel ends in -ext
-					const channelInfo = await this.slack.getChannelInfo(channel)
-					if (channelInfo.channel.name.endsWith('-ext')) {
-						await this.addUserToChannel(emailAddress, channel, true)
-					}
-					else {
-						this.smLogger.error(`Cannot invite external user ${emailAddress} (${user}) to restricted channel ${channel}`)
-					}
-				}
-				else {
-					const errorMsg = `Error inviting ${emailAddress} (${user}) to channel <#${channel}>: ${result.error}`
-					this.smLogger.error(errorMsg)
-					postToRobomergeAlerts(errorMsg)
-				}
 			}
-		}
-		else {
-			const errorMsg = `Unable to add ${emailAddress} to channel <#${channel}>`
-			this.smLogger.error(errorMsg)
-			postToRobomergeAlerts(errorMsg)
+			else {
+				const errorMsg = `Unable to add ${emailAddress} to channel <#${channel}>`
+				this.smLogger.error(errorMsg)
+				postToRobomergeAlerts(errorMsg)
+			}
 		}
 	}
 
