Notify user when doing a reconsider does not result in any action

diff --git a/Engine/Extras/RoboMerge/v3/src/robo/nodebot.ts b/Engine/Extras/RoboMerge/v3/src/robo/nodebot.ts
--- a/Engine/Extras/RoboMerge/v3/src/robo/nodebot.ts
+++ b/Engine/Extras/RoboMerge/v3/src/robo/nodebot.ts
@@ -1251,6 +1251,20 @@ export class NodeBot extends PerforceStatefulBot implements NodeBotInterface {
 
 	onAlreadyIntegrated(event: AlreadyIntegrated) {
 		this.conflicts.onAlreadyIntegrated(event)
+
+		if (event.change.userRequest && this.slackMessages) {
+			if (this.slackMessages) {
+				const message = `Change ${event.change.cl} was not necessary in ${event.action.branch.name}`
+				let dm: SlackMessage = {
+					text: message,
+					style: SlackMessageStyles.WARNING,
+					channel: '',
+					mrkdwn: true
+				}
+				
+				this.slackMessages.postDM(this.findEmail(event.change.owner!), event.change.cl, this.branch, dm)			
+			}
+		}
 	}
 
 	onEndIntegratingToGate(event: EndIntegratingToGateEvent) {
diff --git a/Engine/Extras/RoboMerge/v3/src/robo/notifications.ts b/Engine/Extras/RoboMerge/v3/src/robo/notifications.ts
--- a/Engine/Extras/RoboMerge/v3/src/robo/notifications.ts
+++ b/Engine/Extras/RoboMerge/v3/src/robo/notifications.ts
@@ -287,9 +287,12 @@ export class SlackMessages {
 		}
 	}
 
-	async postDM(emailAddress: string|null, cl: number, branchArg: BranchArg, dm: SlackMessage, persistMessage = true) {
+	async postDM(emailAddress: string|Promise<string|null>|null, cl: number, branchArg: BranchArg, dm: SlackMessage, persistMessage = true) {
 		// The Slack API requires a user ID to open a direct message with users.
 		// The most consistent way to do this is getting their email address out of P4.
+		if (emailAddress && typeof(emailAddress) !== 'string') {
+			emailAddress = await emailAddress
+		}
 		if (!emailAddress) {
 			console.error("Failed to get email address during notifications for CL " + cl)
 			return
