Node resolver can be a group

diff --git a/Engine/Extras/RoboMerge/v3/src/common/perforce.ts b/Engine/Extras/RoboMerge/v3/src/common/perforce.ts
--- a/Engine/Extras/RoboMerge/v3/src/common/perforce.ts
+++ b/Engine/Extras/RoboMerge/v3/src/common/perforce.ts
@@ -1228,9 +1228,12 @@ export class PerforceContext {
 
 	// get the email (according to P4) for a specific user
 	async getEmail(username: string) {
-		const output = await this._execP4(null, ['user', '-o', username]);
-		// look for the email field
-		let m = output.match(/\nEmail:\s+([^\n]+)\n/);
+		let m = null
+		if (!username.startsWith('@')) {
+			const output = await this._execP4(null, ['user', '-o', username]);
+			// look for the email field
+			m = output.match(/\nEmail:\s+([^\n]+)\n/);
+		}
 		return m && m[1];
 	}
 
diff --git a/Engine/Extras/RoboMerge/v3/src/robo/notifications.ts b/Engine/Extras/RoboMerge/v3/src/robo/notifications.ts
--- a/Engine/Extras/RoboMerge/v3/src/robo/notifications.ts
+++ b/Engine/Extras/RoboMerge/v3/src/robo/notifications.ts
@@ -498,7 +498,7 @@ export class BotNotifications implements BotEventHandler {
 
 		const userEmail = await blockage.ownerEmail
 		const slackUser = userEmail ? await this.slackMessages.getSlackUser(userEmail) : null
-		const channelPing = slackUser ? `<@${slackUser}>` : `@${blockage.owner}`
+		const channelPing = slackUser ? `<@${slackUser}>` : blockage.owner.startsWith('@') ? blockage.owner : `@${blockage.owner}`
 
 		const isBotUser = isUserAKnownBot(blockage.owner)
 		const text =
