Waiting for CIS can now be made a link to the tab with the job that is being waited on (whatever you want really)

diff --git a/Engine/Extras/RoboMerge/v3/public/js/boilerplate.js b/Engine/Extras/RoboMerge/v3/public/js/boilerplate.js
--- a/Engine/Extras/RoboMerge/v3/public/js/boilerplate.js
+++ b/Engine/Extras/RoboMerge/v3/public/js/boilerplate.js
@@ -1158,7 +1158,8 @@ function preRenderStatusCell_Edge(nodeData, edgeData) {
 					msg += ` (opens: ${(new Date(timestamp)).toLocaleString()})`
 				}
 			}
-			$('<div>').addClass('status-msg').text(msg).appendTo(statusCell)
+			let statusMsg = edgeData.waitingForCISLink ? $(`<a href="${edgeData.waitingForCISLink}">`).prop('target', '_blank') : $("<div>")
+			statusMsg.addClass('status-msg').text(msg).appendTo(statusCell)
 		}
 	}
 
diff --git a/Engine/Extras/RoboMerge/v3/src/robo/branchdefs.ts b/Engine/Extras/RoboMerge/v3/src/robo/branchdefs.ts
--- a/Engine/Extras/RoboMerge/v3/src/robo/branchdefs.ts
+++ b/Engine/Extras/RoboMerge/v3/src/robo/branchdefs.ts
@@ -86,6 +86,7 @@ export type IntegrationWindowPane = {
 
 export const commonOptionFieldsPrototype = {
 	lastGoodCLPath: 0 as string | number,
+	waitingForCISLink: "",
 	pauseCISUnlessAtGate: false,
 
 	initialCL: 0,
diff --git a/Engine/Extras/RoboMerge/v3/src/robo/branchgraph.ts b/Engine/Extras/RoboMerge/v3/src/robo/branchgraph.ts
--- a/Engine/Extras/RoboMerge/v3/src/robo/branchgraph.ts
+++ b/Engine/Extras/RoboMerge/v3/src/robo/branchgraph.ts
@@ -84,6 +84,7 @@ if (botname === '__TEST__') {
 			['additionalSlackChannelForBlockages', 'additionalSlackChannel', null],
 			['postMessagesToAdditionalChannelOnly', 'postOnlyToAdditionalChannel', null],
 			['lastGoodCLPath', 'lastGoodCLPath', null],
+			['waitingForCISLink', 'waitingForCISLink', null],
 			['disallowSkip', 'disallowSkip', null],
 			['incognitoMode', 'incognitoMode', null],
 			['excludeAuthors', 'excludeAuthors', 'override'],
diff --git a/Engine/Extras/RoboMerge/v3/src/robo/gate.ts b/Engine/Extras/RoboMerge/v3/src/robo/gate.ts
--- a/Engine/Extras/RoboMerge/v3/src/robo/gate.ts
+++ b/Engine/Extras/RoboMerge/v3/src/robo/gate.ts
@@ -396,6 +396,10 @@ export class Gate {
 			outStatus.gateClosedMessage = closedMessage
 		}
 
+		if (this.context.options.waitingForCISLink) {
+			outStatus.waitingForCISLink = this.context.options.waitingForCISLink
+		}
+
 		if (this.nextWindowOpenTime) {
 			outStatus.nextWindowOpenTime = this.nextWindowOpenTime
 		}
diff --git a/Engine/Extras/RoboMerge/v3/src/robo/status-types.ts b/Engine/Extras/RoboMerge/v3/src/robo/status-types.ts
--- a/Engine/Extras/RoboMerge/v3/src/robo/status-types.ts
+++ b/Engine/Extras/RoboMerge/v3/src/robo/status-types.ts
@@ -117,6 +117,7 @@ export type EdgeStatusFields = BotStatusFields & {
 	lastGoodCL?: number
 	lastGoodCLJobLink?: string
 	lastGoodCLDate?: Date
+	waitingForCISLink?: string
 
 	num_changes_remaining: number
 }
