Fix stream spec regex to allow for *

diff --git a/Engine/Extras/RoboMerge/v3/src/new/graph.ts b/Engine/Extras/RoboMerge/v3/src/new/graph.ts
--- a/Engine/Extras/RoboMerge/v3/src/new/graph.ts
+++ b/Engine/Extras/RoboMerge/v3/src/new/graph.ts
@@ -60,7 +60,7 @@ function makeBranchId(arg: string | Branch) {
 		depotPath = arg as string
 	}
 
-	if (!depotPath.match(/^\/\/[\w-.][\/\w-.]*$/)) {
+	if (!depotPath.match(/^\/\/[\w-.][\/\w-.\*]*$/)) {
 		throw new Error('Invalid stream path: ' + depotPath)
 	}
 	return depotPath as Stream
