Do not implicitly apply the stream name as an alias to a node using subPaths

diff --git a/Engine/Extras/RoboMerge/v3/src/robo/branchdefs.ts b/Engine/Extras/RoboMerge/v3/src/robo/branchdefs.ts
--- a/Engine/Extras/RoboMerge/v3/src/robo/branchdefs.ts
+++ b/Engine/Extras/RoboMerge/v3/src/robo/branchdefs.ts
@@ -425,7 +425,7 @@ export class BranchDefs {
 				}
 			}
 
-			if (def.streamName && !defaultConfigForWholeBot.noStreamAliases) {
+			if (def.streamName && !def.streamSubpath && !defaultConfigForWholeBot.noStreamAliases) {
 				addAlias(upperName, def.streamName.toUpperCase())
 			}
 
diff --git a/Engine/Extras/RoboMerge/v3/src/robo/branchgraph.ts b/Engine/Extras/RoboMerge/v3/src/robo/branchgraph.ts
--- a/Engine/Extras/RoboMerge/v3/src/robo/branchgraph.ts
+++ b/Engine/Extras/RoboMerge/v3/src/robo/branchgraph.ts
@@ -322,7 +322,7 @@ if (botname === '__TEST__') {
 // temporarily. Probably worth rewriting entirely, given current knowledge.
 
 		// automatically alias streamName if specified
-		if (options.streamName && !this.config.noStreamAliases) {
+		if (options.streamName && !options.streamSubpath && !this.config.noStreamAliases) {
 			this.alias(name, options.streamName)
 		}
 
