Clear the workspace before each integrate

diff --git a/Engine/Extras/RoboMerge/v3/src/robo/edgebot.ts b/Engine/Extras/RoboMerge/v3/src/robo/edgebot.ts
--- a/Engine/Extras/RoboMerge/v3/src/robo/edgebot.ts
+++ b/Engine/Extras/RoboMerge/v3/src/robo/edgebot.ts
@@ -528,6 +528,7 @@ class EdgeBotImpl extends PerforceStatefulBot {
 		const doVirtualMerge = !info.userRequest && !info.forceCreateAShelf && !target.flags.has('manual') 
 
 		this.currentIntegrationStartTimestamp = Date.now()
+		await this.p4.sync(info.targetWorkspaceOverride, this.targetBranch.rootPath + '#0', {edgeServerAddress})
 		const [mode, results] = await this.p4.integrate(info.targetWorkspaceOverride, source, changenum, integTarget, {edgeServerAddress, virtual: doVirtualMerge})
 
 		const pending: PendingChange = {change: info, action: target, newCl: changenum}
