Can now specify macros at the node level that can be unique or override the macro defined at the bot level

diff --git a/Engine/Extras/RoboMerge/v3/src/robo/branchdefs.ts b/Engine/Extras/RoboMerge/v3/src/robo/branchdefs.ts
--- a/Engine/Extras/RoboMerge/v3/src/robo/branchdefs.ts
+++ b/Engine/Extras/RoboMerge/v3/src/robo/branchdefs.ts
@@ -48,6 +48,7 @@ const branchBasePrototype = {
 	flowsTo: [''],
 	forceFlowTo: [''],
 	defaultFlow: [''],
+	macros: {} as { [name: string]: string[] } | undefined,
 	resolver: '' as string | null,
 	triager: '' as string | null,
 	nagWhenBlocked: false as boolean | null,
diff --git a/Engine/Extras/RoboMerge/v3/src/robo/branchgraph.ts b/Engine/Extras/RoboMerge/v3/src/robo/branchgraph.ts
--- a/Engine/Extras/RoboMerge/v3/src/robo/branchgraph.ts
+++ b/Engine/Extras/RoboMerge/v3/src/robo/branchgraph.ts
@@ -235,6 +235,7 @@ if (botname === '__TEST__') {
 			notify: (options.notify || []).concat(this.config.globalNotify),
 			forceFlowTo: [], // will compute
 			defaultFlow: options.defaultFlow || [],
+			macros: options.macros,
 			blockAssetTargets: new Set<string>(),
 			enabled: !options.disabled,
 			allowDeadend: !options.disallowDeadend,
diff --git a/Engine/Extras/RoboMerge/v3/src/robo/nodebot.ts b/Engine/Extras/RoboMerge/v3/src/robo/nodebot.ts
--- a/Engine/Extras/RoboMerge/v3/src/robo/nodebot.ts
+++ b/Engine/Extras/RoboMerge/v3/src/robo/nodebot.ts
@@ -1658,13 +1658,27 @@ export class NodeBot extends PerforceStatefulBot implements NodeBotInterface {
 			change.desc = '<description not available>'
 		}
 
+		let combinedMacros: { [name: string]: string[] } = {}
+		for (let macroName in this.branch.macros)
+		{
+			combinedMacros[macroName.toLowerCase()] = this.branch.macros[macroName]
+		}
+		for (let macroName in this.branchGraph.config.macros)
+		{
+			const macroNameLower = macroName.toLowerCase()
+			if (combinedMacros[macroNameLower] === undefined)
+			{
+				combinedMacros[macroNameLower] = this.branchGraph.config.macros[macroName]
+			}
+		}
+
 		const parse = (lines: string[]) => parseDescriptionLines({
 			lines,
 			isDefaultBot: this.branch.isDefaultBot,
 			graphBotName: this.graphBotName,
 			cl: change.change,
 			aliasesUpper: this.branchGraph.config.aliases.map(s => s.toUpperCase()),
-			macros: this.branchGraph.config.macros,
+			macros: combinedMacros,
 			logger: this.nodeBotLogger
 		})
 
