Fix error previewing a graph if a stream hasn't yet been created

diff --git a/Engine/Extras/RoboMerge/v3/src/robo/branchdefs.ts b/Engine/Extras/RoboMerge/v3/src/robo/branchdefs.ts
--- a/Engine/Extras/RoboMerge/v3/src/robo/branchdefs.ts
+++ b/Engine/Extras/RoboMerge/v3/src/robo/branchdefs.ts
@@ -267,7 +267,7 @@ export class BranchDefs {
 		}
 	}
 
-	static parseAndValidate(outErrors: string[], branchSpecsText: string, allStreamSpecs: StreamSpecs): ParseResult {
+	static parseAndValidate(outErrors: string[], branchSpecsText: string, allStreamSpecs: StreamSpecs, isPreviewing?: boolean): ParseResult {
 		const defaultConfigForWholeBot: BotConfig = {
 			defaultStreamDepot: null,
 			defaultIntegrationMethod: null,
@@ -382,7 +382,7 @@ export class BranchDefs {
 				def.streamSubpath
 			)
 
-			if (streamResult.stream && !allStreamSpecs.has(streamResult.stream)) {
+			if (!isPreviewing && streamResult.stream && !allStreamSpecs.has(streamResult.stream)) {
 				outErrors.push(`Stream ${streamResult.stream} not found`)
 			}
 		}
diff --git a/Engine/Extras/RoboMerge/v3/src/robo/preview.ts b/Engine/Extras/RoboMerge/v3/src/robo/preview.ts
--- a/Engine/Extras/RoboMerge/v3/src/robo/preview.ts
+++ b/Engine/Extras/RoboMerge/v3/src/robo/preview.ts
@@ -40,7 +40,7 @@ export async function getPreview(cl: number, singleBot?: string) {
 		const fileText = await p4.print(`${path}@=${cl}`)
 
 		const validationErrors: string[] = []
-		const result = BranchDefs.parseAndValidate(validationErrors, fileText, allStreamSpecs)
+		const result = BranchDefs.parseAndValidate(validationErrors, fileText, allStreamSpecs, true)
 
 		const errorPrefix = `\n\t${bot} validation failed: `
 		if (!result.branchGraphDef) {
