When using a resolver group which is the same name as a swarm group avoid spamming reviews for every change

diff --git a/Engine/Extras/RoboMerge/v3/src/robo/edgebot.ts b/Engine/Extras/RoboMerge/v3/src/robo/edgebot.ts
--- a/Engine/Extras/RoboMerge/v3/src/robo/edgebot.ts
+++ b/Engine/Extras/RoboMerge/v3/src/robo/edgebot.ts
@@ -360,9 +360,13 @@ class EdgeBotImpl extends PerforceStatefulBot {
 		description += '\n\n' // description has been trimmed
 
 		// if the owner is specifically overridden (can be by reconsider, resolver or manual tag)
-		const overriddenOwner = getIntegrationOwner(this.branch, info.branch, info.owner)
+		let overriddenOwner = getIntegrationOwner(this.branch, info.branch, info.owner)
 
 		if (overriddenOwner) {
+			// Need to avoid swarm notifying same named groups
+			if (overriddenOwner.startsWith('@')) {
+				overriddenOwner = `@ ${overriddenOwner.substring(1)}`
+			}
 			description += `#ROBOMERGE-OWNER: ${overriddenOwner}\n`
 		}
 
