Disable error until we can look closer at the node merging when rootpaths don't match

diff --git a/Engine/Extras/RoboMerge/v3/src/robo/targets.ts b/Engine/Extras/RoboMerge/v3/src/robo/targets.ts
--- a/Engine/Extras/RoboMerge/v3/src/robo/targets.ts
+++ b/Engine/Extras/RoboMerge/v3/src/robo/targets.ts
@@ -454,7 +454,7 @@ export function processOtherBotTargets(
 		const routes = ubergraph.findAllRoutesBetween(sourceNode, targetNode, [], new Set(skipNodes))
 		if (routes.length === 0)
 		{
-			errors.push(`No route between '${sourceNode.debugName}' and '${targetNode.debugName}'`)
+			//errors.push(`No route between '${sourceNode.debugName}' and '${targetNode.debugName}'`)
 			continue
 		}
 
