don't atify an already atted user (usually a group)

diff --git a/Engine/Extras/RoboMerge/v3/src/robo/notifications.ts b/Engine/Extras/RoboMerge/v3/src/robo/notifications.ts
--- a/Engine/Extras/RoboMerge/v3/src/robo/notifications.ts
+++ b/Engine/Extras/RoboMerge/v3/src/robo/notifications.ts
@@ -93,7 +93,7 @@ export async function postMessageToChannel(message: string, channel: string, sty
 // make a Slack notifcation link for a user
 function atifyUser(user: string) {
 	// don't @ names of bots (currently making them tt style for Slack)
-	return isUserAKnownBot(user) ? `\`${user}\`` : '@' + user
+	return isUserAKnownBot(user) ? `\`${user}\`` : user.startsWith('@') ? user : '@' + user
 }
 
 function generatePersistedSlackMessageKey(sourceCl: number, targetBranchArg: BranchArg, channel: string) {
