Fix example URL

diff --git a/Engine/Extras/RoboMerge/v3/public/allbots.html b/Engine/Extras/RoboMerge/v3/public/allbots.html
--- a/Engine/Extras/RoboMerge/v3/public/allbots.html
+++ b/Engine/Extras/RoboMerge/v3/public/allbots.html
@@ -52,7 +52,7 @@
 				</tbody>
 			</table>
 
-			<p>e.g. <tt>http://robomerge/allbots/12345?bots=test1,test2&noGroups&showOnlyForced</tt></p>
+			<p>e.g. <tt>http://robomerge/allbots?bots=test1,test2&noGroups&showOnlyForced</tt></p>
 		</div>
 	</div>
 </body>
