Secure commands

diff --git a/Engine/Extras/RoboMerge/v3/src/new/graph.ts b/Engine/Extras/RoboMerge/v3/src/new/graph.ts
--- a/Engine/Extras/RoboMerge/v3/src/new/graph.ts
+++ b/Engine/Extras/RoboMerge/v3/src/new/graph.ts
@@ -4,6 +4,7 @@ import { setDefault, sortBy } from '../common/helper';
 import { ContextualLogger } from '../common/logger';
 import { PerforceContext } from '../common/perforce';
 import { Branch, BranchGraphInterface } from '../robo/branch-interfaces';
+import { Status } from '../robo/status';
 
 export type TargetName = string & { __targetBrand: any }
 export type Stream = string & { __streamBrand: any }
@@ -193,10 +194,21 @@ export class Graph {
 		return result
 	}
 
-	dump() {
+	dump(tags?: Set<string>, logger?: ContextualLogger) {
+
+		let validBots: string[] = []
+		if (tags) {
+			validBots = Array.from(this.branchGraphAliases)
+			                 .filter(([_,v]) => Status.includeBranch(v.config.visibility, tags, logger))
+			                 .map(([k,_]) => k)
+		}
+
 		const nodeNames = new Map<Node, TargetName[]>()
 		for (const [name, node] of this.targetNames) {
-			setDefault(nodeNames, node, []).push(name)
+			const botName = name.substring(0,name.search(':'))
+			if (!tags || validBots.includes(botName)) {
+				setDefault(nodeNames, node, []).push(name)
+			}
 		}
 
 		const orderedNodes = [...nodeNames.keys()]
@@ -208,7 +220,9 @@ export class Graph {
 			info.aliases.sort()
 
 			for (const edge of this.getEdgesBySource(node)) {
-				info.edges.push({target: edge.target.debugName, flags: [...edge.flags].join(', ')})
+				if (!tags || validBots.includes(edge.bot)) {
+					info.edges.push({target: edge.target.debugName, flags: [...edge.flags].join(', ')})
+				}
 			}
 			result.push([node.debugName, info])
 		}
@@ -824,11 +838,11 @@ export class Trace {
 		if (outEdges.size === 0)
 			return 'CHANGE_WILL_NOT_REACH_TARGET'
 
-// will have to check if the source node is served by multiple bots, and if so, calculate targets for each of them
-// for now, just pick one
-const arbitraryEdge = outEdges.values().next().value
-const botname = arbitraryEdge.bot
-const isDefaultBot = arbitraryEdge.flags.has('general')
+		// will have to check if the source node is served by multiple bots, and if so, calculate targets for each of them
+		// for now, just pick one
+		const arbitraryEdge = outEdges.values().next().value
+		const botname = arbitraryEdge.bot
+		const isDefaultBot = arbitraryEdge.flags.has('general')
 
 		const relevantTokens = tokensAndBots
 			.filter(([_, bot]) => (bot === botname || isDefaultBot && bot === 'default'))
diff --git a/Engine/Extras/RoboMerge/v3/src/robo/ipc.ts b/Engine/Extras/RoboMerge/v3/src/robo/ipc.ts
--- a/Engine/Extras/RoboMerge/v3/src/robo/ipc.ts
+++ b/Engine/Extras/RoboMerge/v3/src/robo/ipc.ts
@@ -56,8 +56,8 @@ export class IPC {
 			case 'getp4tasks': return { ...OPERATION_SUCCESS, data: getRunningPerforceCommands() } 
 			case 'getWorkspaces': return this.getWorkspaces(msg.args![0] as string)
 			case 'getPersistence': return this.getPersistence(msg.args![0] as string)
-			case 'traceRoute': return this.traceRoute(msg.args![0] as Query)
-			case 'dumpGraph': return this.dumpGraph()
+			case 'traceRoute': return this.traceRoute(msg.args![0] as Query, msg.args![1] as Set<string>)
+			case 'dumpGraph': return this.dumpGraph(msg.args![0] as Set<string>)
 			case 'getIsRunning': return this.isRunning(msg.args![0] as string)
 			case 'restartBot': return this.restartBot(msg.args![0] as string, msg.args![1] as string)
 			case 'crashGraphBot': return this.crashGraphBot(msg.args![0] as string, msg.args![1] as string)
@@ -164,7 +164,7 @@ export class IPC {
 		return {statusCode: 400, message: `Unknown bot '${botname}'`}
 	}
 
-	private async traceRoute(query: Query): Promise<OperationReturnType> {
+	private async traceRoute(query: Query, tags?: Set<string>): Promise<OperationReturnType> {
 
 		const cl = parseInt(query.cl)
 		if (isNaN(cl)) {
@@ -180,13 +180,24 @@ export class IPC {
 				return { statusCode: 400, message: routeOrError } 
 
 			const result: any[] = []
-			for (const edge of routeOrError) {
+			for (let edgeIdx=0; edgeIdx < routeOrError.length; edgeIdx++) {
+				let edge = routeOrError[edgeIdx]
 				const bot = edge.sourceAnnotation as NodeBotInterface
-				result.push({
-					name: bot.branchGraph.botname + ':' + bot.branch.name,
-					stream: edge.source.stream,
-					lastCl: bot.lastCl
-				})
+				let canShowResult = true
+				if (tags) {
+					canShowResult = Status.includeBranch(bot.branchGraph.config.visibility, tags, this.ipcLogger)
+				}
+				if (canShowResult) {
+					result.push({
+						name: bot.branchGraph.botname + ':' + bot.branch.name,
+						stream: edge.source.stream,
+						lastCl: bot.lastCl
+					})
+				} else if (edgeIdx == 0) {
+					return { statusCode: 400, message: "UNKNOWN_SOURCE_BRANCH" } 
+				} else if (edgeIdx == routeOrError.length - 1) {
+					return { statusCode: 400, message: "UNKNOWN_TARGET_BRANCH" } 
+				}
 			}
 
 			success = true
@@ -200,8 +211,8 @@ export class IPC {
 		}
 	}
 
-	private dumpGraph(): OperationReturnType {
-		return { ...OPERATION_SUCCESS, data: this.robo.graph.graph.dump() } 
+	private dumpGraph(tags: Set<string>): OperationReturnType {
+		return { ...OPERATION_SUCCESS, data: this.robo.graph.graph.dump(tags, this.ipcLogger) } 
 	}
 
 	private isRunning(botname: string): OperationReturnType {
diff --git a/Engine/Extras/RoboMerge/v3/src/robo/roboserver.ts b/Engine/Extras/RoboMerge/v3/src/robo/roboserver.ts
--- a/Engine/Extras/RoboMerge/v3/src/robo/roboserver.ts
+++ b/Engine/Extras/RoboMerge/v3/src/robo/roboserver.ts
@@ -455,7 +455,7 @@ class RoboWebApp implements AppInterface {
 		return await this.sendMessage('doNodeOp', [botname, nodeName, nodeOp, this.getQueryFromSecure()])
 	}
 
-	@SecureHandler('GET', '/api/p4tasks')
+	@SecureHandler('GET', '/api/p4tasks', {requiredTags: ['fte']})
 	async getP4Tasks() {
 		return (await this.sendMessage('getp4tasks')).data || []
 	}
@@ -515,13 +515,17 @@ class RoboWebApp implements AppInterface {
 		return result.data
 	}
 
-	@Handler('GET', '/api/trace-route')
+	@SecureHandler('GET', '/api/trace-route')
 	async traceRoute() {
-		const query = this.getQuery()
+		if (!this.authData) {
+			throw new Error('Secure call but no auth data?')
+		}
+
+		const query = this.getQueryFromSecure()
 		if (!query.cl || !query.from || !query.to) {
 			return {statusCode: 400, message: '"cl", "from" and "to" query arguments required'}
 		}
-		const result = await this.sendMessage('traceRoute', [query])
+		const result = await this.sendMessage('traceRoute', [query, this.authData.tags])
 		return result.data ? {...result, success: true, route: result.data} :
 			{...result, success: false, code: result.message}
 			
@@ -529,7 +533,10 @@ class RoboWebApp implements AppInterface {
 
 	@SecureHandler('GET', '/debug/dump-graph')
 	async dumpGraph() {
-		return (await this.sendMessage('dumpGraph')).data
+		if (!this.authData) {
+			throw new Error('Secure call but no auth data?')
+		}
+		return (await this.sendMessage('dumpGraph', [this.authData.tags])).data
 	}
 
 	// https://localhost:4433/op/acknowledge?bot=TEST&branch=Main&cl=1237983421
@@ -666,7 +673,7 @@ class RoboWebApp implements AppInterface {
 		return queryObj
 	}
 
-	private getQuery() : {[key: string]: string} {
+	/*private getQuery() : {[key: string]: string} {
 		if (this.authData) {
 			throw new Error('do not use getQuery for secure calls')
 		}
@@ -676,7 +683,7 @@ class RoboWebApp implements AppInterface {
 			queryObj[key] = val
 		}
 		return queryObj
-	}
+	}*/
 
 	getSentryUser() : Sentry.User {
 		if (!this.authData) {
diff --git a/Engine/Extras/RoboMerge/v3/src/robo/status.ts b/Engine/Extras/RoboMerge/v3/src/robo/status.ts
--- a/Engine/Extras/RoboMerge/v3/src/robo/status.ts
+++ b/Engine/Extras/RoboMerge/v3/src/robo/status.ts
@@ -2,7 +2,7 @@
 
 import { ContextualLogger } from '../common/logger';
 import { Branch } from './branch-interfaces';
-import { BranchDefForStatus, BranchStatus, GraphBotState, UserStatusData } from './status-types';
+import { BranchStatus, GraphBotState, UserStatusData } from './status-types';
 import { OperationReturnType } from './ipc';
 import { AuthData } from './session';
 
@@ -49,7 +49,7 @@ export class Status {
 		result.botStates = botStates
 		for (const status of this.allBranches) {
 			// @todo: check for list of tags
-			if (this.includeBranch(status.def, user.tags)) {
+			if (Status.includeBranch(status.def.visibility, user.tags, this.statusLogger)) {
 				branches.push(status)
 				includedBotNames.add(status.bot)
 			}
@@ -84,22 +84,24 @@ export class Status {
 		return result
 	}
 
-	private includeBranch(branch: BranchDefForStatus, userPrivileges: Set<string>) {
-		if (branch.visibility === 'all') {
+	static includeBranch(branchVisibility: string | string[], userPrivileges: Set<string>, logger?: ContextualLogger) {
+		if (branchVisibility === 'all') {
 			return true
 		}
 
-		if (branch.visibility === 'none') {
+		if (branchVisibility === 'none') {
 			return false
 		}
 
-		if (!Array.isArray(branch.visibility)) {
+		if (!Array.isArray(branchVisibility)) {
 			// for now, only keywords 'all' and 'none' supported
-			this.statusLogger.warn('Unknown visibility keyword: ' + branch.visibility)
+			if (logger) {
+				logger.warn('Unknown visibility keyword: ' + branchVisibility)
+			}
 			return false
 		}
 
-		for (const vis of branch.visibility) {
+		for (const vis of branchVisibility) {
 			if (userPrivileges.has(vis)) {
 				// @todo filter flow lists if individual branch bot filtered
 				return true
