Add path to DescribeResults

diff --git a/Engine/Extras/RoboMerge/v3/src/common/perforce.ts b/Engine/Extras/RoboMerge/v3/src/common/perforce.ts
--- a/Engine/Extras/RoboMerge/v3/src/common/perforce.ts
+++ b/Engine/Extras/RoboMerge/v3/src/common/perforce.ts
@@ -75,6 +75,7 @@ export interface DescribeResult {
 	user: string
 	status: string
 	description: string
+	path: string
 	entries: DescribeEntry[]
 	date: Date | null
 }
@@ -1259,6 +1260,7 @@ export class PerforceContext {
 			user: clInfo.user || '',
 			status: clInfo.status || '',
 			description: clInfo.desc || '',
+			path: clInfo.path || ztagResult[1].path,
 			date: clInfo.time ? new Date(clInfo.time * 1000) : null,
 			entries: []
 		};
