Can now start the server in previewOnly mode which will not process any changes, but otherwise brings the server fully up

diff --git a/Engine/Extras/RoboMerge/v3/src/robo/graphbot.ts b/Engine/Extras/RoboMerge/v3/src/robo/graphbot.ts
--- a/Engine/Extras/RoboMerge/v3/src/robo/graphbot.ts
+++ b/Engine/Extras/RoboMerge/v3/src/robo/graphbot.ts
@@ -99,7 +99,7 @@ export class GraphBot implements GraphInterface, BotEventHandler {
 		return branch && branch.bot ? branch.bot as NodeBot : undefined
 	}
 
-	initBots(ubergraph: GraphAPI) {
+	initBots(ubergraph: GraphAPI, previewMode: boolean) {
 		this.eventTriggers = new BotEventTriggers(this.branchGraph.botname, this.branchGraph.config)
 		this.eventTriggers.registerHandler(this)
 		const blockageUrlGenerator: NodeOpUrlGenerator = (blockage : Blockage | null) => { 
@@ -131,7 +131,7 @@ export class GraphBot implements GraphInterface, BotEventHandler {
 		for (const branch of this.branchGraph.branches) {
 			if (branch.enabled) {
 				const persistence = this.settings.getContext(branch.upperName)
-				branch.bot = new NodeBot(branch, this.mailer, this.slackMessages, this.externalUrl, this.eventTriggers, persistence, ubergraph,
+				branch.bot = new NodeBot(branch, previewMode, this.mailer, this.slackMessages, this.externalUrl, this.eventTriggers, persistence, ubergraph,
 					async () => {
 						const errPair = await this.handleRequestedIntegrationsForAllNodes()
 						if (errPair) {
diff --git a/Engine/Extras/RoboMerge/v3/src/robo/nodebot.ts b/Engine/Extras/RoboMerge/v3/src/robo/nodebot.ts
--- a/Engine/Extras/RoboMerge/v3/src/robo/nodebot.ts
+++ b/Engine/Extras/RoboMerge/v3/src/robo/nodebot.ts
@@ -86,7 +86,8 @@ export class NodeBot extends PerforceStatefulBot implements NodeBotInterface {
 	public isActive = false
 
 	constructor(
-		branchDef: Branch, 
+		branchDef: Branch,
+		private readonly previewMode: boolean,
 		mailer: Mailer,
 		slackMessages: SlackMessages | undefined,
 		externalUrl: string, 
@@ -244,6 +245,9 @@ export class NodeBot extends PerforceStatefulBot implements NodeBotInterface {
 
 	/** @return true if full tick completed, only used for analytics */
 	async tick() {
+		if (this.previewMode) {
+			return false
+		}
 		for (const edgeBot of this.edges.values()) {
 			await edgeBot.tick()
 		}
@@ -353,6 +357,10 @@ export class NodeBot extends PerforceStatefulBot implements NodeBotInterface {
 	}
 
 	async processQueuedChange(fromQueue: QueuedChange) {
+		if (this.previewMode) {
+			this._log_action(`(PREVIEW MODE) - Skipping manually queued change ${fromQueue.cl} on ${this.fullName}, requested by ${fromQueue.who}`)
+			return
+		}
 		let logMessage = `Processing manually queued change ${fromQueue.cl} on ${this.fullName}, requested by ${fromQueue.who}`
 		if (fromQueue.workspace) {
 			logMessage += `, in workspace ${fromQueue.workspace}`
diff --git a/Engine/Extras/RoboMerge/v3/src/robo/robo.ts b/Engine/Extras/RoboMerge/v3/src/robo/robo.ts
--- a/Engine/Extras/RoboMerge/v3/src/robo/robo.ts
+++ b/Engine/Extras/RoboMerge/v3/src/robo/robo.ts
@@ -128,6 +128,13 @@ const COMMAND_LINE_ARGS: {[param: string]: Arg<any>} = {
 		dflt: false
 	},
 
+	previewOnly: {
+		match: /^(-previewOnly)$/,
+		parse: str => str === "false" ? false : true,
+		env: 'ROBO_PREVIEW_ONLY',
+		dflt: false
+	},
+
 	// Sentry environment designation -- use 'PROD' to enable Sentry bug tracking
 	epicEnv: {
 		match: /^(-epicEnv)$/,
@@ -355,7 +362,7 @@ async function _onBranchSpecReloaded(graphBot: GraphBot, logger: ContextualLogge
 		--specReloadEntryCount
 	}
 
-	graphBot.initBots(robo.graph)
+	graphBot.initBots(robo.graph, args.previewOnly)
 
 	if (specReloadEntryCount === 0) {
 		// regenerate ubergraph (last update to finish does regen if multiple in flight)
@@ -436,7 +443,7 @@ async function init(logger: ContextualLogger) {
 	const graph = new Graph
 	robo.graph = new GraphAPI(graph)
 	for (const graphBot of robo.graphBots.values()) {
-		graphBot.initBots(robo.graph)
+		graphBot.initBots(robo.graph, args.previewOnly)
 		addBranchGraph(graph, graphBot.branchGraph)
 	}
 }
diff --git a/Engine/Extras/RoboMerge/v3/src/robo/roboserver.ts b/Engine/Extras/RoboMerge/v3/src/robo/roboserver.ts
--- a/Engine/Extras/RoboMerge/v3/src/robo/roboserver.ts
+++ b/Engine/Extras/RoboMerge/v3/src/robo/roboserver.ts
@@ -34,6 +34,12 @@ let ENVIRONMENT: {[param: string]: any}
 			match: /^-devModeUser=(.+)$/,
 			env: 'ROBO_DEV_MODE_USER',
 			dflt: ''
+		},
+		previewOnly: {
+			match: /^(-previewOnly)$/,
+			parse: str => str === "false" ? false : true,
+			env: 'ROBO_PREVIEW_ONLY',
+			dflt: false
 		}
 	}
 
@@ -47,6 +53,9 @@ let ENVIRONMENT: {[param: string]: any}
 if (ENVIRONMENT.devMode) {
 	roboserverStartupLogger.warn('Running in DEV_MODE')
 }
+if (ENVIRONMENT.previewMode) {
+	roboserverStartupLogger.warn('Running in PREVIEW_MODE')
+}
 
 function readUtf8File(path: string) {
 	return new Promise<string>((done, _fail) => 
