Add integrated command

diff --git a/Engine/Extras/RoboMerge/v3/src/common/perforce.ts b/Engine/Extras/RoboMerge/v3/src/common/perforce.ts
--- a/Engine/Extras/RoboMerge/v3/src/common/perforce.ts
+++ b/Engine/Extras/RoboMerge/v3/src/common/perforce.ts
@@ -302,6 +302,11 @@ export interface IntegrateOpts {
 	virtual?: boolean
 }
 
+export interface IntegratedOpts {
+	intoOnly?: boolean
+	startCL?: number
+}
+
 export interface SyncParams {
 	opts?: string[]
 	edgeServerAddress?: string
@@ -1357,6 +1362,37 @@ export class PerforceContext {
 		}
 	}
 
+	async integrated(roboWorkspace: RoboWorkspace, depotPath: string, opts?: IntegratedOpts)
+	{
+		let args = ['integrated']
+		if (opts) {
+			if (opts.intoOnly) {
+				args.push('--into-only')
+			}
+			if (opts.startCL) {
+				args.push('-s')
+				args.push(opts.startCL.toString())
+			}
+		}
+		args.push(depotPath)
+
+		try {
+			return await this._execP4Ztag(roboWorkspace, args)
+		}
+		catch (reason) {
+			if (!isExecP4Error(reason)) {
+				throw reason
+			}
+
+			let [err, output] = reason
+			// If perforce doesn't detect revisions in the given range, return an empty set of revisions
+			if (output.includes("no file(s) integrated.")) {
+				return []
+			}
+			throw err
+		}
+	}
+
 	print(path: string) {
 		return this._execP4(null, ['print', '-q', path], { noCwd: true, quiet: true })
 	}
