Track change feature

diff --git a/Engine/Extras/RoboMerge/v3/public/js/trackchange.js b/Engine/Extras/RoboMerge/v3/public/js/trackchange.js
new file mode 100644
--- /dev/null
+++ b/Engine/Extras/RoboMerge/v3/public/js/trackchange.js
@@ -0,0 +1,44 @@
+// Copyright Epic Games, Inc. All Rights Reserved.
+"use strict"
+
+function generateChangeList(data) {
+
+	let dataObj = data.data
+	console.log(dataObj)
+	let html = '<div style="margin: auto; width: 80%;"><table class="table"><tbody>'
+	for (const cl in dataObj.changes) {
+		console.log()
+		html += '<tr>'
+		html += `<td><b>${dataObj.changes[cl].streamDisplayName}</b></td>`
+		html += `<td><a href="https://p4-swarm.epicgames.net/changes/${cl}" target="_blank">CL#${cl}</a></td>`
+		html += '</tr>'
+	}
+	html += "</tbody></table></div>"
+
+	return html
+}
+
+function doit(query) {
+	$.get(query)
+	.then(data => {
+		const $successPanel = $('#success-panel');
+		$('#changes', $successPanel).html(generateChangeList(data));
+		$('#success-panel').show();
+		receivedTrackingResults()
+	})
+	.catch(error => {
+		const $errorPanel = $('#error-panel');
+		const errorMsg = error.responseText
+			? error.responseText.replace(/\t/g, '    ')
+			: 'Internal error: ' + error.message;
+		$('pre', $errorPanel).text(errorMsg);
+		$errorPanel.show();
+		receivedTrackingResults()
+	});
+}
+
+function receivedTrackingResults(message) {
+	$('#loadingDiv').fadeOut("fast", "swing", function() { $('#changes').fadeIn("fast") })
+}
+
+window.doTrackChange = doit;
diff --git a/Engine/Extras/RoboMerge/v3/public/trackchange.html b/Engine/Extras/RoboMerge/v3/public/trackchange.html
new file mode 100644
--- /dev/null
+++ b/Engine/Extras/RoboMerge/v3/public/trackchange.html
@@ -0,0 +1,69 @@
+<!DOCTYPE html>
+<meta charset="utf-8">
+<html>
+<head>
+	<title>ROBOMERGE track change</title>
+	<link rel="icon" href="/img/robo-r.png" type="image/png">
+
+	<script src="/js/jquery-3.4.1.min.js"></script>
+	<script src="/js/bootstrap.bundle.min.js"></script>
+	<script src="/js/hpcc-wasm-index.min.js"></script>
+	<script src="/js/flow.js"></script>
+	<script src="/js/trackchange.js"></script>
+
+	<link rel="stylesheet" href="/css/bootstrap.min.css">
+	<link rel="stylesheet" href="/css/main.css">
+
+	<script>
+		$(() => {
+			doTrackChange('{{ &query }}');
+		});
+
+	</script>
+</head>
+<body>
+	<header><img src="/img/logo.png"></header>
+	<div class="container">
+
+		<div id="loadingDiv" class="card card-single border-dark loading-box">
+			<div class="text-center card-body">
+				<h1 id="loadingText" class="align-self-center">Tracking change...</h1>
+				<img id="loadingimg" class="mx-auto d-block" style="width:128px;" src="/img/robo-r-rotating.gif">
+			</div>
+		</div>
+		
+		<div id="error-panel" style="display:none;"> <!-- show or hide in script -->
+			<div class="alert alert-warning" role="alert">
+Errors attempting to track change <a href="https://p4-swarm.epicgames.net/changes/{{ cl }}" target="_blank">CL#{{ cl }}</a>
+			</div>
+			<pre></pre>
+		</div>
+		<div id="success-panel" style="display:none;">
+			<div id="graph"></div>
+			<div id="changes" style="display:none;"></div>
+		</div>
+
+		<div style="transform: scale(0.7);">
+			<h4>Options</h4>
+			<table class="table">
+				<thead><tr><th>Setting</th><th>Effect</th></tr></thead>
+				<tbody>
+					<tr>
+						<td style="vertical-align: middle"><tt>bots=</tt></td>
+						<td>
+							Comma separated list of bots to filter results by<br>
+							e.g. <tt>http://robomerge/trackchange/12345?bots=test1,test2
+						</td>
+					</tr>
+					<tr>
+						<td style="vertical-align: middle"><tt>streams=</tt></td>
+						<td>
+							Comma separated list of streams to show. Supports * as a wildcard. Overrides bots<br>
+							e.g. <tt>http://robomerge/trackchange/12345?streams=//UE5/Dev-*,//UE5/Main</tt>
+						</td>
+					</tr>
+				</tbody>
+			</table>
+		</div>
+	</div>
+</body>
\ No newline at end of file
diff --git a/Engine/Extras/RoboMerge/v3/src/robo/ipc.ts b/Engine/Extras/RoboMerge/v3/src/robo/ipc.ts
--- a/Engine/Extras/RoboMerge/v3/src/robo/ipc.ts
+++ b/Engine/Extras/RoboMerge/v3/src/robo/ipc.ts
@@ -1,7 +1,7 @@
 // Copyright Epic Games, Inc. All Rights Reserved.
 
 import { ContextualLogger, isNpmLogLevel, NpmLogLevelCompare, NpmLogLevelValues } from '../common/logger';
-import { getRunningPerforceCommands } from '../common/perforce';
+import { DescribeResult, getRunningPerforceCommands } from '../common/perforce';
 import { Trace } from '../new/graph';
 import { IPCControls, EdgeBotInterface, NodeBotInterface } from './bot-interfaces';
 import { OperationResult } from './branch-interfaces';
@@ -57,6 +57,7 @@ export class IPC {
 			case 'getWorkspaces': return this.getWorkspaces(msg.args![0] as string)
 			case 'getPersistence': return this.getPersistence(msg.args![0] as string)
 			case 'traceRoute': return this.traceRoute(msg.args![0] as Query, msg.args![1] as Set<string>)
+			case 'trackChange': return this.trackChange(msg.args![0], msg.args![1] as Set<string>)
 			case 'dumpGraph': return this.dumpGraph(msg.args![0] as Set<string>)
 			case 'getIsRunning': return this.isRunning(msg.args![0] as string)
 			case 'restartBot': return this.restartBot(msg.args![0] as string, msg.args![1] as string)
@@ -164,6 +165,163 @@ export class IPC {
 		return {statusCode: 400, message: `Unknown bot '${botname}'`}
 	}
 
+	private async trackChange(url: any, userTags: Set<string>): Promise<OperationReturnType> {
+
+		const clStr = url.searchParams.get('cl')
+		if (!clStr) {
+			return {statusCode: 400, message: 'No CL parameter provided.'}
+		}
+
+		const cl = parseInt(clStr)
+		if (isNaN(cl)) {
+			return {statusCode: 400, message: `Invalid CL parameter: ${cl}`}
+		}
+
+		const streamFilter = (() => {
+			const streamsParam: string|undefined = url.searchParams.get('streams');
+			return (streamsParam ? streamsParam.toUpperCase().replace('*','.*').split(',').map(s => new RegExp(s)) : [])
+		})() 
+		const botFilter = (() => { 
+			const botsParam = url.searchParams.get('bots')
+			// botFilter is ignored if streamFilter specified
+			return (streamFilter.length == 0 && botsParam ? botsParam.toUpperCase().split(',') : [])
+		})()
+
+		const graph = this.robo.graph.graph
+
+		const getStreamFromPath = (path: string) => {
+			const match = path.match(/\/\/[^\/]*\/[^\/]*/)
+			return (match ? match[0] : null)
+		}
+
+		const getNode = (desc: DescribeResult) => {
+			if (desc.path) {
+				const stream = getStreamFromPath(desc.path)
+				if (stream) {
+					return graph.findNodeForStream(stream)
+				}
+			}
+			return null
+		}
+
+		const getMergeMethod = (desc: string) => {
+			if (desc.includes("#ROBOMERGE-CONFLICT")) {
+				return 'Merge w/ conflict resolve'
+			} else if (desc.includes("#ROBOMERGE-SOURCE")) {
+				return 'Automerge'
+			} else if (desc.includes("Populate")) {
+				return 'Populate'
+			} else {
+				return 'Manual merge'
+			}
+		}
+
+		let changes = new Map<number, any>()
+
+		let gatherCLInfo = async (cl: number, lastCL?: number) => {
+			let desc = await this.robo.p4.describe(cl, 1)
+			const clNode = getNode(desc)
+			if (clNode || !lastCL|| desc.description.includes("#ROBOMERGE-SOURCE")) {
+				changes.set(cl, {desc, node: clNode, destCLs: (lastCL ? [lastCL] : []) })
+				return true
+			}
+			return false
+		}
+
+		await gatherCLInfo(cl)
+
+		let data: any = {originalCL: cl, changes: {}}
+
+		const sourceMatch = (changes.get(data.originalCL).desc as DescribeResult).description.match(/#ROBOMERGE-SOURCE: (.*)\n/g)
+		if (sourceMatch) {
+			let lastCL = cl
+			const matches = Array.from(sourceMatch[0].matchAll(/CL (\d+)/g))
+			for (let i=matches.length-1; i>=0;i--) {
+				let sourceCL = parseInt(matches[i][1])
+				await gatherCLInfo(sourceCL, lastCL)
+				lastCL = sourceCL
+				if (i == 0) {
+					data.originalCL = sourceCL
+				}
+			}
+		}
+
+		const isFTE = userTags.has("fte")
+
+		let clsToConsider: number[] = [data.originalCL]
+		while (clsToConsider.length > 0) {
+			const clToConsider = clsToConsider[0]
+			clsToConsider = clsToConsider.slice(1)
+
+			let changeToConsider = changes.get(clToConsider)
+
+			let includeInResults = false
+			let hasAutomergeTarget = false
+			let streamDisplayName = ""
+			if (changeToConsider.node) {
+				let edges = graph.getEdgesBySource(changeToConsider.node)
+				for (let edge of edges) {
+					if (!includeInResults) {
+						const bot = edge.sourceAnnotation as NodeBotInterface
+						if (botFilter.length == 0 || botFilter.includes(bot.branchGraph.botname)) {
+							if (Status.includeBranch(bot.branchGraph.config.visibility, userTags, this.ipcLogger)) {
+								streamDisplayName = changeToConsider.node.stream
+								includeInResults = true
+							}
+						}
+					}
+					if (edge.flags.has('automatic')) {
+						hasAutomergeTarget = true
+					}
+					if (includeInResults && hasAutomergeTarget) {
+						break
+					}
+				}
+			} else if (isFTE) {
+				includeInResults = true
+				streamDisplayName = getStreamFromPath(changeToConsider.desc.path) || changeToConsider.desc.path
+			}
+
+			if (streamDisplayName.length == 0) {
+				streamDisplayName = "//****/****"
+			}
+
+			if (includeInResults && streamFilter.length > 0) {
+				includeInResults = streamFilter.some(re => streamDisplayName.toUpperCase().match(re))
+			}
+
+			for (let i=0; i < changeToConsider.desc.entries.length; i++) {
+				const entry = changeToConsider.desc.entries[i]
+				const integrated = await this.robo.p4.integrated(null, entry.depotFile, {intoOnly: true, startCL: clToConsider})
+				if (integrated.length > 0) {
+					for (let integ of integrated) {
+						if (integ.endToRev == `#${entry.rev}`)
+						{
+							if (!changes.has(integ.change)) {
+								changeToConsider.destCLs.push(integ.change)
+								await gatherCLInfo(integ.change)
+							}
+						}
+					}
+					break
+				} else if (hasAutomergeTarget && changeToConsider.desc.entries.length == 1) {
+					// If we only have 1 entry and we didn't get integration info off of it
+					// and the graph suggests we are expecting there could be other changes
+					// get the full describe results
+					changeToConsider.desc = await this.robo.p4.describe(clToConsider)
+				}
+			}
+			clsToConsider = clsToConsider.concat(changeToConsider.destCLs)
+
+			if (includeInResults) {
+				const mergeMethod = clToConsider != data.originalCL ? getMergeMethod(changeToConsider.desc.description) : ""
+				data.changes[`${clToConsider}`] = {streamDisplayName, mergeMethod}
+			}
+		}
+
+		return {...OPERATION_SUCCESS, data}
+	}
+
 	private async traceRoute(query: Query, tags?: Set<string>): Promise<OperationReturnType> {
 
 		const cl = parseInt(query.cl)
diff --git a/Engine/Extras/RoboMerge/v3/src/robo/roboserver.ts b/Engine/Extras/RoboMerge/v3/src/robo/roboserver.ts
--- a/Engine/Extras/RoboMerge/v3/src/robo/roboserver.ts
+++ b/Engine/Extras/RoboMerge/v3/src/robo/roboserver.ts
@@ -176,6 +176,39 @@ class RoboWebApp implements AppInterface {
 		return readUtf8File('public/index.html')
 	}
 
+	@SecureHandler('GET', '/trackchange/*', {filetype: 'text/html'}) 
+	async trackChange(clStr: string) {
+		if (!this.authData) {
+			throw new Error('Secure call but no auth data?')
+		}
+
+		const cl = parseInt(clStr)
+		if (isNaN(cl)) {
+			throw new Error(`Failed to parse alleged CL '${clStr}'`)
+		}
+
+		let query = "/trackchange"
+		if (this.request.url.search.length > 0) {
+			query += `${this.request.url.search}&`
+		} else {
+			query += "?"
+		}
+		query += `cl=${clStr}`
+
+		const template = await readUtf8File('public/trackchange.html')
+		return Mustache.render(template, {cl, query})
+	}
+
+	@SecureHandler('GET', '/trackchange') 
+	getTrackedChanges() {
+		if (!this.authData) {
+			throw new Error('Secure call but no auth data?')
+		}
+
+		return this.sendMessage('trackChange', [this.request.url, this.authData.tags])
+	}
+
+
 	@Handler('GET', '/preview/*', {filetype: 'text/html'}) 
 	previewNoBot(clStr: string) {
 		return this.preview(clStr)
