Add waitingForCISURL link details

diff --git a/Engine/Extras/RoboMerge/v3/README.md b/Engine/Extras/RoboMerge/v3/README.md
--- a/Engine/Extras/RoboMerge/v3/README.md
+++ b/Engine/Extras/RoboMerge/v3/README.md
@@ -111,6 +111,7 @@ for both node and edge.
 |`initialCL`                |e |            | |First run only: which CL start _after_   | |
 |`isDefaultBot`             |e |`false`     | |Run plain #robomerge commands? Should be `false` for streams monitored by multiple bots| |
 |`lastGoodCLPath`           |e |            | |'Gate' file to read to find CIS-approved CL| |
+|`waitingForCISLink`        |e |            | |URL to link to for waiting for CIS messages| |
 |`maxFilesPerIntegration`   |e |`-1`        | |Reject integrations with more files than this| |
 |`notify`                   |e |            | |Additional people to email on blockages  |Also `globalNotify`|
 |`resolver`                 |e |            | |Single designated resolver               |Currently applies to both source and target nodes|
