Change how params are sent to track change to be more consistent with preview in hopes that it works on the production server as searchParams did not
Fix bots filter not working to remove streams not represented in the graph

diff --git a/Engine/Extras/RoboMerge/v3/src/robo/ipc.ts b/Engine/Extras/RoboMerge/v3/src/robo/ipc.ts
--- a/Engine/Extras/RoboMerge/v3/src/robo/ipc.ts
+++ b/Engine/Extras/RoboMerge/v3/src/robo/ipc.ts
@@ -165,9 +165,9 @@ export class IPC {
 		return {statusCode: 400, message: `Unknown bot '${botname}'`}
 	}
 
-	private async trackChange(url: any, userTags: Set<string>): Promise<OperationReturnType> {
+	private async trackChange(queryObj: any, userTags: Set<string>): Promise<OperationReturnType> {
 
-		const clStr = url.searchParams.get('cl')
+		const clStr = queryObj.cl
 		if (!clStr) {
 			return {statusCode: 400, message: 'No CL parameter provided.'}
 		}
@@ -178,11 +178,11 @@ export class IPC {
 		}
 
 		const streamFilter = (() => {
-			const streamsParam: string|undefined = url.searchParams.get('streams');
+			const streamsParam: string|undefined = queryObj.streams;
 			return (streamsParam ? streamsParam.toUpperCase().replace('*','.*').split(',').map(s => new RegExp(s)) : [])
 		})() 
 		const botFilter = (() => { 
-			const botsParam = url.searchParams.get('bots')
+			const botsParam = queryObj.bots
 			// botFilter is ignored if streamFilter specified
 			return (streamFilter.length == 0 && botsParam ? botsParam.toUpperCase().split(',') : [])
 		})()
@@ -277,7 +277,7 @@ export class IPC {
 						break
 					}
 				}
-			} else if (isFTE) {
+			} else if (isFTE && botFilter.length == 0) {
 				includeInResults = true
 				streamDisplayName = getStreamFromPath(changeToConsider.desc.path) || changeToConsider.desc.path
 			}
diff --git a/Engine/Extras/RoboMerge/v3/src/robo/roboserver.ts b/Engine/Extras/RoboMerge/v3/src/robo/roboserver.ts
--- a/Engine/Extras/RoboMerge/v3/src/robo/roboserver.ts
+++ b/Engine/Extras/RoboMerge/v3/src/robo/roboserver.ts
@@ -205,7 +205,13 @@ class RoboWebApp implements AppInterface {
 			throw new Error('Secure call but no auth data?')
 		}
 
-		return this.sendMessage('trackChange', [this.request.url, this.authData.tags])
+		const queryObj: {[key: string]: string} = {}
+		for (const [key, val] of this.request.url.searchParams) {
+			queryObj[key] = val
+		}
+
+
+		return this.sendMessage('trackChange', [queryObj, this.authData.tags])
 	}
 
 
