bots and streamnames can be used at the same time to filter

diff --git a/Engine/Extras/RoboMerge/v3/public/trackchange.html b/Engine/Extras/RoboMerge/v3/public/trackchange.html
--- a/Engine/Extras/RoboMerge/v3/public/trackchange.html
+++ b/Engine/Extras/RoboMerge/v3/public/trackchange.html
@@ -58,7 +58,7 @@ Errors attempting to track change <a href="https://p4-swarm.epicgames.net/change
 					<tr>
 						<td style="vertical-align: middle"><tt>streams=</tt></td>
 						<td>
-							Comma separated list of streams to show. Supports * as a wildcard. Overrides bots<br>
+							Comma separated list of streams to show. Supports * as a wildcard.<br>
 							e.g. <tt>http://robomerge/trackchange/12345?streams=//UE5/Dev-*,//UE5/Main</tt>
 						</td>
 					</tr>
diff --git a/Engine/Extras/RoboMerge/v3/src/robo/ipc.ts b/Engine/Extras/RoboMerge/v3/src/robo/ipc.ts
--- a/Engine/Extras/RoboMerge/v3/src/robo/ipc.ts
+++ b/Engine/Extras/RoboMerge/v3/src/robo/ipc.ts
@@ -183,8 +183,7 @@ export class IPC {
 		})() 
 		const botFilter = (() => { 
 			const botsParam = queryObj.bots
-			// botFilter is ignored if streamFilter specified
-			return (streamFilter.length == 0 && botsParam ? botsParam.toUpperCase().split(',') : [])
+			return (botsParam ? botsParam.toUpperCase().split(',') : [])
 		})()
 
 		const graph = this.robo.graph.graph
