Fix failure when not specifying the origin CL and the origin CL is not represented in the robomerge graph (or is using streamSubPath)

diff --git a/Engine/Extras/RoboMerge/v3/src/robo/ipc.ts b/Engine/Extras/RoboMerge/v3/src/robo/ipc.ts
--- a/Engine/Extras/RoboMerge/v3/src/robo/ipc.ts
+++ b/Engine/Extras/RoboMerge/v3/src/robo/ipc.ts
@@ -217,10 +217,12 @@ export class IPC {
 
 		let changes = new Map<number, any>()
 
+		let data: any = {originalCL: cl, changes: {}}
+
 		let gatherCLInfo = async (cl: number, lastCL?: number) => {
 			let desc = await this.robo.p4.describe(cl, 1)
 			const clNode = getNode(desc)
-			if (clNode || !lastCL|| desc.description.includes("#ROBOMERGE-SOURCE")) {
+			if (clNode || cl == data.originalCL || desc.description.includes("#ROBOMERGE-SOURCE")) {
 				changes.set(cl, {desc, node: clNode, destCLs: (lastCL ? [lastCL] : []) })
 				return true
 			}
@@ -229,19 +231,17 @@ export class IPC {
 
 		await gatherCLInfo(cl)
 
-		let data: any = {originalCL: cl, changes: {}}
-
 		const sourceMatch = (changes.get(data.originalCL).desc as DescribeResult).description.match(/#ROBOMERGE-SOURCE: (.*)\n/g)
 		if (sourceMatch) {
 			let lastCL = cl
 			const matches = Array.from(sourceMatch[0].matchAll(/CL (\d+)/g))
 			for (let i=matches.length-1; i>=0;i--) {
 				let sourceCL = parseInt(matches[i][1])
-				await gatherCLInfo(sourceCL, lastCL)
-				lastCL = sourceCL
 				if (i == 0) {
 					data.originalCL = sourceCL
 				}
+				await gatherCLInfo(sourceCL, lastCL)
+				lastCL = sourceCL
 			}
 		}
 
