Marshall userTags differently to workaround production webserver issues

diff --git a/Engine/Extras/RoboMerge/v3/src/robo/ipc.ts b/Engine/Extras/RoboMerge/v3/src/robo/ipc.ts
--- a/Engine/Extras/RoboMerge/v3/src/robo/ipc.ts
+++ b/Engine/Extras/RoboMerge/v3/src/robo/ipc.ts
@@ -56,9 +56,9 @@ export class IPC {
 			case 'getp4tasks': return { ...OPERATION_SUCCESS, data: getRunningPerforceCommands() } 
 			case 'getWorkspaces': return this.getWorkspaces(msg.args![0] as string)
 			case 'getPersistence': return this.getPersistence(msg.args![0] as string)
-			case 'traceRoute': return this.traceRoute(msg.args![0] as Query, msg.args![1] as Set<string>)
-			case 'trackChange': return this.trackChange(msg.args![0], msg.args![1] as Set<string>)
-			case 'dumpGraph': return this.dumpGraph(msg.args![0] as Set<string>)
+			case 'traceRoute': return this.traceRoute(msg.args![0] as Query, msg.args![1])
+			case 'trackChange': return this.trackChange(msg.args![0], msg.args![1])
+			case 'dumpGraph': return this.dumpGraph(msg.args![0])
 			case 'getIsRunning': return this.isRunning(msg.args![0] as string)
 			case 'restartBot': return this.restartBot(msg.args![0] as string, msg.args![1] as string)
 			case 'crashGraphBot': return this.crashGraphBot(msg.args![0] as string, msg.args![1] as string)
@@ -165,7 +165,7 @@ export class IPC {
 		return {statusCode: 400, message: `Unknown bot '${botname}'`}
 	}
 
-	private async trackChange(queryObj: any, userTags: Set<string>): Promise<OperationReturnType> {
+	private async trackChange(queryObj: any, tagsObj: any): Promise<OperationReturnType> {
 
 		const clStr = queryObj.cl
 		if (!clStr) {
@@ -177,6 +177,11 @@ export class IPC {
 			return {statusCode: 400, message: `Invalid CL parameter: ${cl}`}
 		}
 
+		let userTags = new Set<string>()
+		for (const tag in tagsObj) {
+			userTags.add(tag)
+		}
+
 		const streamFilter = (() => {
 			const streamsParam: string|undefined = queryObj.streams;
 			return (streamsParam ? streamsParam.toUpperCase().replace('*','.*').split(',').map(s => new RegExp(s)) : [])
@@ -321,13 +326,21 @@ export class IPC {
 		return {...OPERATION_SUCCESS, data}
 	}
 
-	private async traceRoute(query: Query, tags?: Set<string>): Promise<OperationReturnType> {
+	private async traceRoute(query: Query, tagsObj?: any): Promise<OperationReturnType> {
 
 		const cl = parseInt(query.cl)
 		if (isNaN(cl)) {
 			return {statusCode: 400, message: 'Invalid CL parameter: ' + query.cl}
 		}
 
+		let tags
+		if (tagsObj) {
+			tags = new Set<string>()
+			for (const tag in tagsObj) {
+				tags.add(tag)
+			}
+		}
+
 		const tracer = new Trace(this.robo.graph.graph, this.ipcLogger)
 
 		let success = false
@@ -368,7 +381,11 @@ export class IPC {
 		}
 	}
 
-	private dumpGraph(tags: Set<string>): OperationReturnType {
+	private dumpGraph(tagsObj: any): OperationReturnType {
+		let tags = new Set<string>()
+		for (const tag in tagsObj) {
+			tags.add(tag)
+		}
 		return { ...OPERATION_SUCCESS, data: this.robo.graph.graph.dump(tags, this.ipcLogger) } 
 	}
 
diff --git a/Engine/Extras/RoboMerge/v3/src/robo/roboserver.ts b/Engine/Extras/RoboMerge/v3/src/robo/roboserver.ts
--- a/Engine/Extras/RoboMerge/v3/src/robo/roboserver.ts
+++ b/Engine/Extras/RoboMerge/v3/src/robo/roboserver.ts
@@ -210,8 +210,12 @@ class RoboWebApp implements AppInterface {
 			queryObj[key] = val
 		}
 
+		const tagsObj: {[key: string]: boolean} = {}
+		for (const tag of this.authData.tags) {
+			tagsObj[tag] = true
+		}
 
-		return this.sendMessage('trackChange', [queryObj, this.authData.tags])
+		return this.sendMessage('trackChange', [queryObj, tagsObj])
 	}
 
 
@@ -573,7 +577,13 @@ class RoboWebApp implements AppInterface {
 		if (!query.cl || !query.from || !query.to) {
 			return {statusCode: 400, message: '"cl", "from" and "to" query arguments required'}
 		}
-		const result = await this.sendMessage('traceRoute', [query, this.authData.tags])
+
+		const tagsObj: {[key: string]: boolean} = {}
+		for (const tag of this.authData.tags) {
+			tagsObj[tag] = true
+		}
+
+		const result = await this.sendMessage('traceRoute', [query, tagsObj])
 		return result.data ? {...result, success: true, route: result.data} :
 			{...result, success: false, code: result.message}
 			
@@ -584,7 +594,13 @@ class RoboWebApp implements AppInterface {
 		if (!this.authData) {
 			throw new Error('Secure call but no auth data?')
 		}
-		return (await this.sendMessage('dumpGraph', [this.authData.tags])).data
+
+		const tagsObj: {[key: string]: boolean} = {}
+		for (const tag of this.authData.tags) {
+			tagsObj[tag] = true
+		}
+
+		return (await this.sendMessage('dumpGraph', [tagsObj])).data
 	}
 
 	// https://localhost:4433/op/acknowledge?bot=TEST&branch=Main&cl=1237983421
