Clean up cases where no results are found

diff --git a/Engine/Extras/RoboMerge/v3/public/js/trackchange.js b/Engine/Extras/RoboMerge/v3/public/js/trackchange.js
--- a/Engine/Extras/RoboMerge/v3/public/js/trackchange.js
+++ b/Engine/Extras/RoboMerge/v3/public/js/trackchange.js
@@ -23,7 +23,13 @@ function doit(query) {
 	.then(data => {
 		const $successPanel = $('#success-panel');
 		$('#changes', $successPanel).html(generateChangeList(data));
-		$('#success-panel').show();
+		if (data.data.changes.length > 0) {
+			$successPane.show();
+		} else {
+			const $errorPanel = $('#error-panel');
+			$('pre', $errorPanel).html('No results found.')
+			$errorPanel.show();
+		}
 		receivedTrackingResults()
 	})
 	.catch(error => {
diff --git a/Engine/Extras/RoboMerge/v3/src/robo/ipc.ts b/Engine/Extras/RoboMerge/v3/src/robo/ipc.ts
--- a/Engine/Extras/RoboMerge/v3/src/robo/ipc.ts
+++ b/Engine/Extras/RoboMerge/v3/src/robo/ipc.ts
@@ -283,7 +283,11 @@ export class IPC {
 				}
 			} else if (isFTE && botFilter.length == 0) {
 				includeInResults = true
-				streamDisplayName = getStreamFromPath(changeToConsider.desc.path) || changeToConsider.desc.path
+				if (changeToConsider.desc.path) {
+					streamDisplayName = getStreamFromPath(changeToConsider.desc.path) || changeToConsider.desc.path
+				} else if (changeToConsider.desc.entries.length > 0) {
+					streamDisplayName = getStreamFromPath(changeToConsider.desc.entries[0].depotFile) || streamDisplayName
+				}
 			}
 
 			if (streamDisplayName.length == 0) {
