Fix determining if any changes were returned

diff --git a/Engine/Extras/RoboMerge/v3/public/js/trackchange.js b/Engine/Extras/RoboMerge/v3/public/js/trackchange.js
--- a/Engine/Extras/RoboMerge/v3/public/js/trackchange.js
+++ b/Engine/Extras/RoboMerge/v3/public/js/trackchange.js
@@ -4,7 +4,6 @@
 function generateChangeList(data) {
 
 	let dataObj = data.data
-	console.log(dataObj)
 	let html = '<div style="margin: auto; width: 80%;"><table class="table"><tbody>'
 	for (const cl in dataObj.changes) {
 		console.log()
@@ -23,8 +22,8 @@ function doit(query) {
 	.then(data => {
 		const $successPanel = $('#success-panel');
 		$('#changes', $successPanel).html(generateChangeList(data));
-		if (data.data.changes.length > 0) {
-			$successPane.show();
+		if (Object.keys(data.data.changes).length > 0) {
+			$successPanel.show();
 		} else {
 			const $errorPanel = $('#error-panel');
 			$('pre', $errorPanel).html('No results found.')
