[Track changes]
Fix error when a manual merge node is not represented by any of the bot nodes
Do not check integrated on every file in the changelist if it was not merged via a robomerge method

diff --git a/Engine/Extras/RoboMerge/v3/src/robo/ipc.ts b/Engine/Extras/RoboMerge/v3/src/robo/ipc.ts
--- a/Engine/Extras/RoboMerge/v3/src/robo/ipc.ts
+++ b/Engine/Extras/RoboMerge/v3/src/robo/ipc.ts
@@ -27,6 +27,10 @@ export interface Message {
 // roboserver.ts -- getQueryFromSecure()
 type Query = {[key: string]: string};
 
+const RobomergeMethodStrings = ['initialSubmit','merge_with_conflict','automerge'] as const;
+type RobomergeMethods = typeof RobomergeMethodStrings[number];
+type MergeMethod = RobomergeMethods|'populate'|'manual_merge'
+
 export type OperationReturnType = {
 	statusCode: number
 	message: string // Goal: to provide meaningful error messaging to the end user
@@ -165,6 +169,7 @@ export class IPC {
 		return {statusCode: 400, message: `Unknown bot '${botname}'`}
 	}
 
+
 	private async trackChange(queryObj: any, tagsObj: any): Promise<OperationReturnType> {
 
 		const clStr = queryObj.cl
@@ -208,15 +213,15 @@ export class IPC {
 			return null
 		}
 
-		const getMergeMethod = (desc: string) => {
+		const getMergeMethod = (desc: string): MergeMethod => {
 			if (desc.includes("#ROBOMERGE-CONFLICT")) {
-				return 'Merge w/ conflict resolve'
+				return 'merge_with_conflict'
 			} else if (desc.includes("#ROBOMERGE-SOURCE")) {
-				return 'Automerge'
+				return 'automerge'
 			} else if (desc.includes("Populate")) {
-				return 'Populate'
+				return 'populate'
 			} else {
-				return 'Manual merge'
+				return 'manual_merge'
 			}
 		}
 
@@ -258,6 +263,9 @@ export class IPC {
 			clsToConsider = clsToConsider.slice(1)
 
 			let changeToConsider = changes.get(clToConsider)
+			if (!changeToConsider) {
+				continue
+			}
 
 			let includeInResults = false
 			let hasAutomergeTarget = false
@@ -298,6 +306,8 @@ export class IPC {
 				includeInResults = streamFilter.some(re => streamDisplayName.toUpperCase().match(re))
 			}
 
+			const mergeMethod = clToConsider != data.originalCL ? getMergeMethod(changeToConsider.desc.description) : 'initialSubmit'
+
 			for (let i=0; i < changeToConsider.desc.entries.length; i++) {
 				const entry = changeToConsider.desc.entries[i]
 				const integrated = await this.robo.p4.integrated(null, entry.depotFile, {intoOnly: true, startCL: clToConsider})
@@ -312,7 +322,8 @@ export class IPC {
 						}
 					}
 					break
-				} else if (hasAutomergeTarget && changeToConsider.desc.entries.length == 1) {
+				} else if (hasAutomergeTarget && changeToConsider.desc.entries.length == 1 && 
+							(mergeMethod in RobomergeMethodStrings)) {
 					// If we only have 1 entry and we didn't get integration info off of it
 					// and the graph suggests we are expecting there could be other changes
 					// get the full describe results
@@ -322,7 +333,6 @@ export class IPC {
 			clsToConsider = clsToConsider.concat(changeToConsider.destCLs)
 
 			if (includeInResults) {
-				const mergeMethod = clToConsider != data.originalCL ? getMergeMethod(changeToConsider.desc.description) : ""
 				data.changes[`${clToConsider}`] = {streamDisplayName, mergeMethod}
 			}
 		}
