Add Track Change button

diff --git a/Engine/Extras/RoboMerge/v3/public/js/boilerplate.js b/Engine/Extras/RoboMerge/v3/public/js/boilerplate.js
--- a/Engine/Extras/RoboMerge/v3/public/js/boilerplate.js
+++ b/Engine/Extras/RoboMerge/v3/public/js/boilerplate.js
@@ -484,6 +484,24 @@ function generateRobomergeFooter() {
 	branchesButton.text("Branch Data")
 	fteButtonDiv.append(branchesButton)
 
+	let trackChangeButton = $('<button id="trackChangeButton">')
+	trackChangeButton.addClass("btn btn-sm btn-outline-dark")
+	trackChangeButton.click(function() { 
+		let data = promptFor({
+			cl: `Enter the CL to track:`
+		})
+		if (data) {
+			// Ensure they entered a CL
+			if (isNaN(parseInt(data.cl))) {
+				displayErrorMessage("Please provide a valid changelist number to track.")
+				return
+			}
+			window.open(`/trackchange/${data.cl}`, '_blank') 
+		}
+	})
+	trackChangeButton.text("Track Change")
+	fteButtonDiv.append(trackChangeButton)
+
 	let currentlyRunningDiv = $('<div id="currentlyRunning">')
 	fixedFooterContents.append(currentlyRunningDiv)
 
