More reliably track changes done by manual merge

diff --git a/Engine/Extras/RoboMerge/v3/src/robo/ipc.ts b/Engine/Extras/RoboMerge/v3/src/robo/ipc.ts
--- a/Engine/Extras/RoboMerge/v3/src/robo/ipc.ts
+++ b/Engine/Extras/RoboMerge/v3/src/robo/ipc.ts
@@ -231,9 +231,10 @@ export class IPC {
 
 		let gatherCLInfo = async (cl: number, lastCL?: number) => {
 			let desc = await this.robo.p4.describe(cl, 1)
+			const mergeMethod = cl != data.originalCL ? getMergeMethod(desc.description) : 'initialSubmit'
 			const clNode = getNode(desc)
-			if (clNode || cl == data.originalCL || desc.description.includes("#ROBOMERGE-SOURCE")) {
-				changes.set(cl, {desc, node: clNode, destCLs: (lastCL ? [lastCL] : []) })
+			if (clNode || mergeMethod in RobomergeMethodStrings) {
+				changes.set(cl, {desc, mergeMethod, node: clNode, destCLs: (lastCL ? [lastCL] : []) })
 				return true
 			}
 			return false
@@ -306,14 +307,14 @@ export class IPC {
 				includeInResults = streamFilter.some(re => streamDisplayName.toUpperCase().match(re))
 			}
 
-			const mergeMethod = clToConsider != data.originalCL ? getMergeMethod(changeToConsider.desc.description) : 'initialSubmit'
-
 			for (let i=0; i < changeToConsider.desc.entries.length; i++) {
 				const entry = changeToConsider.desc.entries[i]
 				const integrated = await this.robo.p4.integrated(null, entry.depotFile, {intoOnly: true, startCL: clToConsider})
 				if (integrated.length > 0) {
 					for (let integ of integrated) {
-						if (integ.endToRev == `#${entry.rev}`)
+						const startToRev = parseInt(integ.startToRev.slice(1))
+						const endToRev = parseInt(integ.endToRev.slice(1))
+						if (entry.rev <= endToRev && entry.rev > startToRev)
 						{
 							if (!changes.has(integ.change)) {
 								changeToConsider.destCLs.push(integ.change)
@@ -323,7 +324,7 @@ export class IPC {
 					}
 					break
 				} else if (hasAutomergeTarget && changeToConsider.desc.entries.length == 1 && 
-							(mergeMethod in RobomergeMethodStrings)) {
+							(changeToConsider.mergeMethod in RobomergeMethodStrings)) {
 					// If we only have 1 entry and we didn't get integration info off of it
 					// and the graph suggests we are expecting there could be other changes
 					// get the full describe results
@@ -333,7 +334,7 @@ export class IPC {
 			clsToConsider = clsToConsider.concat(changeToConsider.destCLs)
 
 			if (includeInResults) {
-				data.changes[`${clToConsider}`] = {streamDisplayName, mergeMethod}
+				data.changes[`${clToConsider}`] = {streamDisplayName, mergeMethod: changeToConsider.mergeMethod}
 			}
 		}
 
