Code to show merge method, but commented out for now
Fix incorrectly cached merge method

diff --git a/Engine/Extras/RoboMerge/v3/public/js/trackchange.js b/Engine/Extras/RoboMerge/v3/public/js/trackchange.js
--- a/Engine/Extras/RoboMerge/v3/public/js/trackchange.js
+++ b/Engine/Extras/RoboMerge/v3/public/js/trackchange.js
@@ -1,6 +1,22 @@
 // Copyright Epic Games, Inc. All Rights Reserved.
 "use strict"
 
+function getMergeMethodString(mergeMethod) {
+	switch(mergeMethod) {
+		case "automerge":
+			return "Automerge"
+		case "initialSubmit":
+			return ""
+		case "merge_with_conflict":
+			return "Merge w/ conflict"
+		case "manual_merge":
+			return "Manual merge"
+		case "populate":
+			return "Populate"
+	}
+	return "UNKNOWN CASE: " + mergeMethod
+}
+
 function generateChangeList(data) {
 
 	let dataObj = data.data
@@ -10,6 +26,7 @@ function generateChangeList(data) {
 		html += '<tr>'
 		html += `<td><b>${dataObj.changes[cl].streamDisplayName}</b></td>`
 		html += `<td><a href="https://p4-swarm.epicgames.net/changes/${cl}" target="_blank">CL#${cl}</a></td>`
+		//html += `<td style="text-align:center;">${getMergeMethodString(dataObj.changes[cl].mergeMethod)}</td>`
 		html += '</tr>'
 	}
 	html += "</tbody></table></div>"
diff --git a/Engine/Extras/RoboMerge/v3/src/robo/ipc.ts b/Engine/Extras/RoboMerge/v3/src/robo/ipc.ts
--- a/Engine/Extras/RoboMerge/v3/src/robo/ipc.ts
+++ b/Engine/Extras/RoboMerge/v3/src/robo/ipc.ts
@@ -234,7 +234,7 @@ export class IPC {
 			const mergeMethod = cl != data.originalCL ? getMergeMethod(desc.description) : 'initialSubmit'
 			const clNode = getNode(desc)
 			if (clNode || mergeMethod in RobomergeMethodStrings) {
-				changes.set(cl, {desc, mergeMethod, node: clNode, destCLs: (lastCL ? [lastCL] : []) })
+				changes.set(cl, {desc, node: clNode, destCLs: (lastCL ? [lastCL] : []) })
 				return true
 			}
 			return false
@@ -307,6 +307,10 @@ export class IPC {
 				includeInResults = streamFilter.some(re => streamDisplayName.toUpperCase().match(re))
 			}
 
+			// Can't cache this because the passed in changelist is incorrectly labelled initialSubmit the first time
+			// it is considered
+			const mergeMethod = clToConsider != data.originalCL ? getMergeMethod(changeToConsider.desc.description) : 'initialSubmit'
+
 			for (let i=0; i < changeToConsider.desc.entries.length; i++) {
 				const entry = changeToConsider.desc.entries[i]
 				const integrated = await this.robo.p4.integrated(null, entry.depotFile, {intoOnly: true, startCL: clToConsider})
@@ -324,7 +328,7 @@ export class IPC {
 					}
 					break
 				} else if (hasAutomergeTarget && changeToConsider.desc.entries.length == 1 && 
-							(changeToConsider.mergeMethod in RobomergeMethodStrings)) {
+							(mergeMethod in RobomergeMethodStrings)) {
 					// If we only have 1 entry and we didn't get integration info off of it
 					// and the graph suggests we are expecting there could be other changes
 					// get the full describe results
@@ -334,7 +338,7 @@ export class IPC {
 			clsToConsider = clsToConsider.concat(changeToConsider.destCLs)
 
 			if (includeInResults) {
-				data.changes[`${clToConsider}`] = {streamDisplayName, mergeMethod: changeToConsider.mergeMethod}
+				data.changes[`${clToConsider}`] = {streamDisplayName, mergeMethod}
 			}
 		}
 
