Fix reloadWorkspace race condition

diff --git a/Engine/Extras/RoboMerge/v3/src/common/perforce.ts b/Engine/Extras/RoboMerge/v3/src/common/perforce.ts
--- a/Engine/Extras/RoboMerge/v3/src/common/perforce.ts
+++ b/Engine/Extras/RoboMerge/v3/src/common/perforce.ts
@@ -644,11 +644,11 @@ export class PerforceContext {
 		return null
 	}
 
-	async reloadWorkspace(workspaceName: string) {
+	reloadWorkspace(workspaceName: string) {
 		return this._execP4Ztag(null, ['reload', '-c', workspaceName]);
 	}
 
-	getEdgeServerAddress(serverId: string): Promise<string> {
+	getEdgeServerAddress(serverId: string) {
 		return this._execP4(null, ['-ztag', '-F', '%Address%', 'server', '-o', serverId])
 	}
 
diff --git a/Engine/Extras/RoboMerge/v3/src/robo/robo.ts b/Engine/Extras/RoboMerge/v3/src/robo/robo.ts
--- a/Engine/Extras/RoboMerge/v3/src/robo/robo.ts
+++ b/Engine/Extras/RoboMerge/v3/src/robo/robo.ts
@@ -309,7 +309,7 @@ async function _initWorkspacesForGraphBot(graphBot: GraphBot, existingWorkspaces
 		}
 	}
 
-	Promise.all(reloadedWorkspaces)
+	await Promise.all(reloadedWorkspaces)
 	if (workspacesToReset.length > 0) {
 		logger.info('The following workspaces already exist and will be reset: ' + workspacesToReset.join(', '))
 		await p4util.cleanWorkspaces(robo.p4, workspacesToReset)
