Improve formatting on multiple validation errors when previewing

diff --git a/Engine/Extras/RoboMerge/v3/src/robo/preview.ts b/Engine/Extras/RoboMerge/v3/src/robo/preview.ts
--- a/Engine/Extras/RoboMerge/v3/src/robo/preview.ts
+++ b/Engine/Extras/RoboMerge/v3/src/robo/preview.ts
@@ -39,12 +39,18 @@ export async function getPreview(cl: number, singleBot?: string) {
 		}
 		const fileText = await p4.print(`${path}@=${cl}`)
 
-		const validationErrors: string[] = []
+		let validationErrors: string[] = []
 		const result = BranchDefs.parseAndValidate(validationErrors, fileText, allStreamSpecs, true)
 
 		const errorPrefix = `\n\t${bot} validation failed: `
 		if (!result.branchGraphDef) {
-			errors.push(errorPrefix + (validationErrors.length === 0 ? 'unknown error' : validationErrors.join('')))
+			if (validationErrors.length == 0) {
+				errors.push(errorPrefix + 'unknown error')
+			} else if (validationErrors.length == 1) {
+				errors.push(errorPrefix + validationErrors[0])
+			} else {
+				errors.push(errorPrefix + validationErrors.map(err => `\n\t\t${err}`).join(''))
+			}
 			continue
 		}
 
