TrackChange: Add ability to filter on depot(s)

diff --git a/Engine/Extras/RoboMerge/v3/public/trackchange.html b/Engine/Extras/RoboMerge/v3/public/trackchange.html
--- a/Engine/Extras/RoboMerge/v3/public/trackchange.html
+++ b/Engine/Extras/RoboMerge/v3/public/trackchange.html
@@ -76,6 +76,13 @@ Errors attempting to track change <a href="https://p4-swarm.epicgames.net/change
 							e.g. <tt>http://robomerge/trackchange/12345?bots=test1,test2
 						</td>
 					</tr>
+					<tr>
+						<td style="vertical-align: middle"><tt>depots=</tt></td>
+						<td>
+							Comma separated list of depots to filter results by<br>
+							e.g. <tt>http://robomerge/trackchange/12345?depots=UE5,Stream-Test
+						</td>
+					</tr>
 					<tr>
 						<td style="vertical-align: middle"><tt>streams=</tt></td>
 						<td>
diff --git a/Engine/Extras/RoboMerge/v3/src/robo/ipc.ts b/Engine/Extras/RoboMerge/v3/src/robo/ipc.ts
--- a/Engine/Extras/RoboMerge/v3/src/robo/ipc.ts
+++ b/Engine/Extras/RoboMerge/v3/src/robo/ipc.ts
@@ -195,6 +195,10 @@ export class IPC {
 			const botsParam = queryObj.bots
 			return (botsParam ? botsParam.toUpperCase().split(',') : [])
 		})()
+		const depotFilter = (() => { 
+			const depotsParam: string = queryObj.depots
+			return (depotsParam ? depotsParam.toUpperCase().split(',').map(depot => `//${depot}`) : [])
+		})()
 
 		const graph = this.robo.graph.graph
 
@@ -303,8 +307,11 @@ export class IPC {
 				streamDisplayName = "//****/****"
 			}
 
-			if (includeInResults && streamFilter.length > 0) {
-				includeInResults = streamFilter.some(re => streamDisplayName.toUpperCase().match(re))
+			if (includeInResults)
+			{
+				const upperSteamDisplayName = streamDisplayName.toUpperCase()
+				includeInResults = depotFilter.length == 0 || depotFilter.some(depot => upperSteamDisplayName.startsWith(depot))
+				includeInResults = includeInResults && (streamFilter.length == 0 || streamFilter.some(re => upperSteamDisplayName.match(re)))
 			}
 
 			// Can't cache this because the passed in changelist is incorrectly labelled initialSubmit the first time
