TrackChange: Fix nodes with no outgoing edges not getting included

diff --git a/Engine/Extras/RoboMerge/v3/src/new/graph.ts b/Engine/Extras/RoboMerge/v3/src/new/graph.ts
--- a/Engine/Extras/RoboMerge/v3/src/new/graph.ts
+++ b/Engine/Extras/RoboMerge/v3/src/new/graph.ts
@@ -152,11 +152,7 @@ export class Graph {
 		return this.targetNames.get(name)
 	}
 
-	getEdgesBySource(src: Node, ...flags: EdgeFlag[]) {
-		const edges = this.edgesBySource.get(src)
-		if (!edges)
-			return new Set<Edge>()
-
+	private static filterEdgesByFlags(edges: Set<Edge>, flags: EdgeFlag[]) {
 		if (flags.length === 0)
 			return edges
 
@@ -169,6 +165,20 @@ export class Graph {
 		}))
 	}
 
+	getEdgesBySource(src: Node, ...flags: EdgeFlag[]) {
+		return Graph.filterEdgesByFlags(this.edgesBySource.get(src) || new Set<Edge>, flags)
+	}
+
+	getEdgesByTarget(target: Node, ...flags: EdgeFlag[]) {
+		return Graph.filterEdgesByFlags(this.edgesByTarget.get(target) || new Set<Edge>, flags)
+	}
+
+	getEdgesForNode(node: Node, ...flags: EdgeFlag[]) {
+		const srcEdges = this.edgesBySource.get(node) || new Set<Edge>()
+		const targetEdges = this.edgesByTarget.get(node) || new Set<Edge>()
+		return Graph.filterEdgesByFlags(new Set([...srcEdges, ...targetEdges]), flags)
+	}
+
 	computeReachable(result: Set<Node>, src: Node, ...flags: EdgeFlag[]) {
 		// if result starts off with stuff, those nodes will be ignored (i.e. not checked)
 
diff --git a/Engine/Extras/RoboMerge/v3/src/robo/ipc.ts b/Engine/Extras/RoboMerge/v3/src/robo/ipc.ts
--- a/Engine/Extras/RoboMerge/v3/src/robo/ipc.ts
+++ b/Engine/Extras/RoboMerge/v3/src/robo/ipc.ts
@@ -276,7 +276,7 @@ export class IPC {
 			let hasAutomergeTarget = false
 			let streamDisplayName = ""
 			if (changeToConsider.node) {
-				let edges = graph.getEdgesBySource(changeToConsider.node)
+				let edges = graph.getEdgesForNode(changeToConsider.node)
 				for (let edge of edges) {
 					if (!includeInResults) {
 						const bot = edge.sourceAnnotation as NodeBotInterface
