fix repl command

diff --git a/Engine/Extras/RoboMerge/v3/repl/roborepl.py b/Engine/Extras/RoboMerge/v3/repl/roborepl.py
--- a/Engine/Extras/RoboMerge/v3/repl/roborepl.py
+++ b/Engine/Extras/RoboMerge/v3/repl/roborepl.py
@@ -34,7 +34,7 @@ context = dict(
 
 def get_bot(name):
 	try:
-		return context['_bots'][bot.lower()]
+		return context['_bots'][name.lower()]
 	except KeyError:
 		raise f"Unknown bot '{name}'"
 
