Do not output changelist descriptions in the log (security risk and just unnecessary logging)
Added p4 add and fstat functions
p4 edit can now pass additional arguments
Fix a few cases where CL1234 is output in a way that makes searching logs for the CL# not work easily

diff --git a/Engine/Extras/RoboMerge/v3/src/common/perforce.ts b/Engine/Extras/RoboMerge/v3/src/common/perforce.ts
--- a/Engine/Extras/RoboMerge/v3/src/common/perforce.ts
+++ b/Engine/Extras/RoboMerge/v3/src/common/perforce.ts
@@ -1244,9 +1244,27 @@ export class PerforceContext {
 	}
 
 	/** Check out a file into a specific changelist ** ASYNC ** */
-	edit(roboWorkspace: RoboWorkspace, cl: number, filePath: string) {
+	edit(roboWorkspace: RoboWorkspace, cl: number, filePath: string, additionalArgs?: string[]) {
 		const workspace = coercePerforceWorkspace(roboWorkspace);
-		return this._execP4(workspace, ['edit', '-c', cl.toString(), filePath]);
+		const args = [
+			'edit', 
+			'-c', cl.toString(), 
+			...(additionalArgs || []),
+			filePath
+		]
+		return this._execP4(workspace, args);
+	}
+
+	/** Add a file into a specific changelist ** ASYNC ** */
+	add(roboWorkspace: RoboWorkspace, cl: number, filePath: string, filetype?: string) {
+		const workspace = coercePerforceWorkspace(roboWorkspace);
+		const args = [
+			'add', 
+			'-c', cl.toString(), 
+			...(filetype ? ['-t', filetype] : []),
+			filePath
+		]
+		return this._execP4(workspace, args);
 	}
 
 	async describe(cl: number, maxFiles?: number, includeShelved?: boolean) {
@@ -1312,7 +1330,7 @@ export class PerforceContext {
 
 		// run the P4 change command to update
 		const changeFlag = opts.changeSubmitted ? '-f' : '-u';
-		this.logger.info(`Executing: 'p4 change -i ${changeFlag}' to edit CL${changelist}`);
+		this.logger.info(`Executing: 'p4 change -i ${changeFlag}' to edit CL ${changelist}`);
 
 		await this._execP4(workspace, ['change', '-i', changeFlag],
 				{ stdin: form, quiet: true, edgeServerAddress: opts.edgeServerAddress })
@@ -1362,6 +1380,23 @@ export class PerforceContext {
 		}
 	}
 
+	async fstat(roboWorkspace: RoboWorkspace, depotPath: string) {
+		try {
+			return await this._execP4Ztag(roboWorkspace, ['fstat', depotPath])
+		}
+		catch (reason) {
+			if (!isExecP4Error(reason)) {
+				throw reason
+			}
+
+			let [err, output] = reason
+			if (output.includes('no such file')) {
+				return []
+			}
+			throw err
+		}
+	}
+
 	async integrated(roboWorkspace: RoboWorkspace, depotPath: string, opts?: IntegratedOpts)
 	{
 		let args = ['integrated']
diff --git a/Engine/Extras/RoboMerge/v3/src/robo/edgebot.ts b/Engine/Extras/RoboMerge/v3/src/robo/edgebot.ts
--- a/Engine/Extras/RoboMerge/v3/src/robo/edgebot.ts
+++ b/Engine/Extras/RoboMerge/v3/src/robo/edgebot.ts
@@ -751,8 +751,7 @@ class EdgeBotImpl extends PerforceStatefulBot {
 		this.resetIntegrationTimestamp()
 
 		// log that this was done
-		const log_str = (target.description || '').trim();
-		this.edgeBotLogger.info(`Submitted CL ${finalCl} to ${this.targetBranch.name}\n    ${log_str.replace(/\n/g, "\n    ")}`);
+		this.edgeBotLogger.info(`Submitted CL ${finalCl} to ${this.targetBranch.name}`);
 
 		if (info.author != 'robomerge') {
 			// change owner, so users can edit change descriptions later for reconsideration
