Persistence files are now backed up to perforce (by default once an hour) to avoid truly catastrophic data loss in the case that the robomerge hardware crashes and burns

diff --git a/Engine/Extras/RoboMerge/v3/src/robo/graphbot.ts b/Engine/Extras/RoboMerge/v3/src/robo/graphbot.ts
--- a/Engine/Extras/RoboMerge/v3/src/robo/graphbot.ts
+++ b/Engine/Extras/RoboMerge/v3/src/robo/graphbot.ts
@@ -89,7 +89,7 @@ export class GraphBot implements GraphInterface, BotEventHandler {
 			this.botLogger.error(`Problem starting up bot ${botname}: ${error}`);
 		}
 
-		this.settings = new Settings(botname, this.branchGraph, this.botLogger)
+		this.settings = new Settings(botname, this.branchGraph, this.botLogger, this.p4)
 
 		this.externalUrl = externalUrl
 	}
diff --git a/Engine/Extras/RoboMerge/v3/src/robo/robo.ts b/Engine/Extras/RoboMerge/v3/src/robo/robo.ts
--- a/Engine/Extras/RoboMerge/v3/src/robo/robo.ts
+++ b/Engine/Extras/RoboMerge/v3/src/robo/robo.ts
@@ -11,7 +11,7 @@ import { Arg, readProcessArgs } from '../common/args';
 import { _setTimeout } from '../common/helper';
 import { ContextualLogger } from '../common/logger';
 import { Mailer } from '../common/mailer';
-import { ClientSpec, getPerforceUsername, getRootDirectoryForBranch, initializePerforce, PerforceContext, StreamSpecs } from '../common/perforce';
+import { ClientSpec, getPerforceUsername, getRootDirectoryForBranch, initializePerforce, PerforceContext, StreamSpecs, Workspace } from '../common/perforce';
 import { BuildVersion, VersionReader } from '../common/version';
 import { CertFiles } from '../common/webserver';
 import { addBranchGraph, Graph, GraphAPI } from '../new/graph';
@@ -174,6 +174,24 @@ const COMMAND_LINE_ARGS: {[param: string]: Arg<any>} = {
 		match: /^-persistenceDir=(.+)$/,
 		env: 'ROBO_PERSISTENCE_DIR',
 		dflt: process.platform === 'win32' ? 'D:/ROBO' : path.resolve(process.env.HOME || '/root', '.robomerge')
+	},
+	persistenceBackupPath: {
+		match: /^-persistenceBackupPath=(.*)$/,
+		env: 'ROBO_PERSISTENCE_BACKUP_PATH',
+		dflt: '//GamePlugins/RobomergePersistence'
+	},
+	persistenceBackupWorkspace: {
+		match: /^-bs_workspace=(.+)$/,
+		env: 'ROBO_PERSISTENCE_BACKUP_WORKSPACE',
+		dflt: 'robomerge-persistence-' + os.hostname()
+	},
+	persistenceBackupFrequency: {
+		match: /^-persistenceBackupFrequency=([0-9]+)$/,
+		env: 'ROBO_PERSISTENCE_BACKUP_FREQUENCY',
+		parse: (str: string) => {
+			return parseInt(str)
+		},
+		dflt: 60
 	}
 };
 
@@ -387,10 +405,35 @@ async function _onBranchSpecReloaded(graphBot: GraphBot, logger: ContextualLogge
 
 async function init(logger: ContextualLogger) {
 
-	if (!args.branchSpecsRootPath) {
-		logger.warn('Auto brancher updater not configured!')
+	let lookupStream = async function(workspaceName: Workspace, rootPath: string): Promise<string> {
+
+		const match = rootPath.match(/(\/\/.*?\/.*?)(\/|$)/)
+		if (match) {
+			const stream = match[1]
+
+			while (true) {
+				try {
+					const streams = await robo.p4.streams();
+					if (streams.has(stream)) {
+						break
+					}
+				}
+				catch (err) {
+				}
+
+				const timeout = 5.0;
+				logger.info(`Will check for ${workspaceName} again in ${timeout} sec...`);
+				await _setTimeout(timeout*1000);
+			}
+			return stream
+		}
+		else {
+			logger.warn(`Unable to determine stream from root path '${rootPath}'`)
+			return ""
+		}
 	}
-	else {
+
+	if (args.branchSpecsRootPath) {
 		const branchSpecsAbsPath = fs.realpathSync(args.branchSpecsDirectory)
 		const autoUpdaterConfig = {
 			rootPath: args.branchSpecsRootPath,
@@ -401,37 +444,42 @@ async function init(logger: ContextualLogger) {
 		// Ensure we have a workspace for branch specs
 		const workspace: Object[] = await robo.p4.find_workspace_by_name(args.branchSpecsWorkspace)
 		if (workspace.length === 0) {
-			logger.info("Cannot find branch spec workspace " + args.branchSpecsWorkspace + 
-						", creating a new one.")
-
-			while (true) {
-				const match = args.branchSpecsRootPath.match(/(\/\/.*?\/.*?)(\/|$)/)
-				if (match) {
-					const branchSpecsStream = match[1]
-					try {
-						const streams = await robo.p4.streams();
-						if (streams.has(branchSpecsStream)) {
-							break
-						}
-					}
-					catch (err) {
-					}
+			logger.info(`Cannot find workspace ${args.branchSpecsWorkspace}, creating a new one.`)
 
-					const timeout = 5.0;
-					logger.info(`Will check again in ${timeout} sec...`);
-					await _setTimeout(timeout*1000);
-				}
-				else {
-					logger.warn(`Unable to determine stream from root path '${args.branchSpecsRootPath}'`)
-				}
+			if ((await lookupStream(args.branchSpecsWorkspace, args.branchSpecsRootPath)).length > 0) {
+				await robo.p4.newBranchSpecWorkspace(autoUpdaterConfig.workspace, args.branchSpecsRootPath)
 			}
-
-			await robo.p4.newBranchSpecWorkspace(autoUpdaterConfig.workspace, args.branchSpecsRootPath)
 		}
 
 		// make sure we've got the latest branch specs
 		logger.info('Syncing latest branch specs')
 		await AutoBranchUpdater.init({p4: robo.p4}, autoUpdaterConfig, logger)
+
+	} else {
+		logger.warn('Auto brancher updater not configured!')
+	}
+
+	if (args.persistenceBackupFrequency > 0 && !args.previewOnly) {
+		const workspace: Object[] = await robo.p4.find_workspace_by_name(args.persistenceBackupWorkspace)
+		if (workspace.length === 0) {
+			logger.info(`Cannot find workspace ${args.persistenceBackupWorkspace}, creating a new one.`)
+
+			const workspace = {directory: fs.realpathSync(args.persistenceDir), name: args.persistenceBackupWorkspace}
+			const stream = await lookupStream(workspace, args.persistenceBackupPath)
+			if (stream.length > 0) {
+				let roots: string | string[] = '/app/data' // default linux path
+				if (workspace.directory !== roots) {
+					roots = [roots, workspace.directory] // specified directory
+				}
+				const params: any = {
+					AltRoots: roots,
+					Stream: stream
+				}		
+				await robo.p4.newWorkspace(args.persistenceBackupWorkspace, params)
+			}
+		}
+	} else {
+		logger.warn('Persistence backup not configured!')
 	}
 
 	_initMailer(logger)
diff --git a/Engine/Extras/RoboMerge/v3/src/robo/settings.ts b/Engine/Extras/RoboMerge/v3/src/robo/settings.ts
--- a/Engine/Extras/RoboMerge/v3/src/robo/settings.ts
+++ b/Engine/Extras/RoboMerge/v3/src/robo/settings.ts
@@ -7,6 +7,7 @@ import { VersionReader } from '../common/version';
 import { BranchGraph } from './branchgraph';
 import { PersistentConflict } from './conflict-interfaces';
 import { NOTIFICATIONS_PERSISTENCE_KEY, SLACK_MESSAGES_PERSISTENCE_KEY } from './notifications';
+import { PerforceContext, Workspace } from '../common/perforce';
 import semver = require('semver')
 
 /**
@@ -123,16 +124,21 @@ function getPersistenceFilepath(botname: string) {
 }
 
 export class Settings {
-	filename: string;
-	lokiFilename: string;
-	enableSave: boolean;
+	filename: string
+	enableSave: boolean
 	object: { [key: string]: any }
+	private	lastPersistBackupTime: Date
 	private readonly settingsLogger: ContextualLogger
+	private readonly p4workspace?: Workspace
 
-	constructor(botname: string, branchgraph: BranchGraph, parentLogger: ContextualLogger) {
+	constructor(botname: string, branchgraph: BranchGraph, parentLogger: ContextualLogger, private p4: PerforceContext) {
 		this.settingsLogger = parentLogger.createChild('Settings')
 
+		this.lastPersistBackupTime = new Date(Date.now())
 		this.filename = getPersistenceFilepath(botname)
+		if (args.persistenceBackupFrequency > 0) {
+			this.p4workspace = {directory: fs.realpathSync(args.persistenceDir), name: args.persistenceBackupWorkspace}
+		}
 
 		// see if we should enable saves
 		this.enableSave = !process.env["NOSAVE"];
@@ -216,7 +222,7 @@ export class Settings {
 		 */
 		
 		this.object.version = VERSION.raw
-		this._saveObject();
+		this._saveObject(true);
 	}
 
 	/**
@@ -423,7 +429,7 @@ export class Settings {
 		return this.filename + FilePurposeExtensionMap[purpose]
 	}
 
-	_saveObject() {
+	_saveObject(ensureInP4?: boolean) {
 		if (this.enableSave) {
 
 			const persistentFile = this.getFilename('PERSISTENCE')
@@ -436,7 +442,33 @@ export class Settings {
 			let filebits = JSON.stringify(this.object, null, '  ')
 			fs.writeFileSync(tempFile, filebits, "utf8")
 
-			fs.copyFileSync(tempFile,persistentFile)
+			fs.copyFileSync(tempFile, persistentFile)
+
+			if (this.p4workspace) {
+				if (ensureInP4) {
+					this.p4.fstat(this.p4workspace, persistentFile)
+					.then((fstat) => {
+						if (fstat.length == 0 || fstat[0].headAction == 'delete') {
+							this.p4.new_cl(this.p4workspace!, "Adding persistent backup", [])
+							.then((cl: number) => {
+								this.p4.add(this.p4workspace!, cl, persistentFile, "text+wS64")
+								.then(() => this.p4.submit(this.p4workspace!, cl))
+							})
+						}
+					})
+				}
+				else if ((Date.now() - this.lastPersistBackupTime.getTime()) / 1000 >= args.persistenceBackupFrequency) {
+					this.lastPersistBackupTime = new Date(Date.now())
+
+					this.p4.sync(this.p4workspace, persistentFile, {opts: ['-k']})
+					.then(() => this.p4.new_cl(this.p4workspace!, "Updating persistent backup"))
+					.then((cl: number) => {
+						this.p4.edit(this.p4workspace!, cl, persistentFile, ['-k'])
+						.then(() => this.p4.submit(this.p4workspace!, cl))
+					})
+					
+				}
+			}
 		}
 	}
 }
