All stages of a stomp should be done with virtual integrations

diff --git a/Engine/Extras/RoboMerge/v3/src/robo/edgebot.ts b/Engine/Extras/RoboMerge/v3/src/robo/edgebot.ts
--- a/Engine/Extras/RoboMerge/v3/src/robo/edgebot.ts
+++ b/Engine/Extras/RoboMerge/v3/src/robo/edgebot.ts
@@ -529,7 +529,10 @@ class EdgeBotImpl extends PerforceStatefulBot {
 
 		// We want to do a virtual merge unless the change is going to be (or could be
 		// in userRequest/reconsider case) handed off to the user
-		const doVirtualMerge = !info.userRequest && !info.forceCreateAShelf && !target.flags.has('manual') 
+		const doVirtualMerge = 
+			   info.forceStompChanges 
+			|| info.sendNoShelfNotification 
+		    || (!info.userRequest && !info.forceCreateAShelf && !target.flags.has('manual'))
 
 		this.currentIntegrationStartTimestamp = Date.now()
 		await this.p4.sync(info.targetWorkspaceOverride, this.targetBranch.rootPath + '#0', {edgeServerAddress})
