Disable backup persistence for the functional tests

diff --git a/Engine/Extras/RoboMerge/v3/run-tests.sh b/Engine/Extras/RoboMerge/v3/run-tests.sh
--- a/Engine/Extras/RoboMerge/v3/run-tests.sh
+++ b/Engine/Extras/RoboMerge/v3/run-tests.sh
@@ -75,6 +75,7 @@ docker run -d -p 8877:8877 -p 8811:8811 -p 25:25 -h robomerge_functtest --name r
     -e ROBO_DEV_MODE_USER=testuser1  \
     -e P4PORT=p4docker:1666  \
     -e ROBO_BRANCHSPECS_ROOT_PATH=//RoboMergeData/Main  \
+    -e ROBO_PERSISTENCE_BACKUP_FREQUENCY=0 \
     -e ROBO_NO_MAIL=true \
     -e ROBO_NO_TLS=true  \
     -e ROBO_EXTERNAL_URL=http://robomerge_functtest:8877  \
