Fix lookupStream logging

diff --git a/Engine/Extras/RoboMerge/v3/src/robo/robo.ts b/Engine/Extras/RoboMerge/v3/src/robo/robo.ts
--- a/Engine/Extras/RoboMerge/v3/src/robo/robo.ts
+++ b/Engine/Extras/RoboMerge/v3/src/robo/robo.ts
@@ -405,7 +405,7 @@ async function _onBranchSpecReloaded(graphBot: GraphBot, logger: ContextualLogge
 
 async function init(logger: ContextualLogger) {
 
-	let lookupStream = async function(workspaceName: Workspace, rootPath: string): Promise<string> {
+	let lookupStream = async function(rootPath: string): Promise<string> {
 
 		const match = rootPath.match(/(\/\/.*?\/.*?)(\/|$)/)
 		if (match) {
@@ -422,7 +422,7 @@ async function init(logger: ContextualLogger) {
 				}
 
 				const timeout = 5.0;
-				logger.info(`Will check for ${workspaceName} again in ${timeout} sec...`);
+				logger.info(`Will check for ${stream} again in ${timeout} sec...`);
 				await _setTimeout(timeout*1000);
 			}
 			return stream
@@ -446,7 +446,7 @@ async function init(logger: ContextualLogger) {
 		if (workspace.length === 0) {
 			logger.info(`Cannot find workspace ${args.branchSpecsWorkspace}, creating a new one.`)
 
-			if ((await lookupStream(args.branchSpecsWorkspace, args.branchSpecsRootPath)).length > 0) {
+			if ((await lookupStream(args.branchSpecsRootPath)).length > 0) {
 				await robo.p4.newBranchSpecWorkspace(autoUpdaterConfig.workspace, args.branchSpecsRootPath)
 			}
 		}
@@ -465,7 +465,7 @@ async function init(logger: ContextualLogger) {
 			logger.info(`Cannot find workspace ${args.persistenceBackupWorkspace}, creating a new one.`)
 
 			const workspace = {directory: fs.realpathSync(args.persistenceDir), name: args.persistenceBackupWorkspace}
-			const stream = await lookupStream(workspace, args.persistenceBackupPath)
+			const stream = await lookupStream(args.persistenceBackupPath)
 			if (stream.length > 0) {
 				let roots: string | string[] = '/app/data' // default linux path
 				if (workspace.directory !== roots) {
