Fix unused import

diff --git a/Engine/Extras/RoboMerge/v3/src/robo/robo.ts b/Engine/Extras/RoboMerge/v3/src/robo/robo.ts
--- a/Engine/Extras/RoboMerge/v3/src/robo/robo.ts
+++ b/Engine/Extras/RoboMerge/v3/src/robo/robo.ts
@@ -11,7 +11,7 @@ import { Arg, readProcessArgs } from '../common/args';
 import { _setTimeout } from '../common/helper';
 import { ContextualLogger } from '../common/logger';
 import { Mailer } from '../common/mailer';
-import { ClientSpec, getPerforceUsername, getRootDirectoryForBranch, initializePerforce, PerforceContext, StreamSpecs, Workspace } from '../common/perforce';
+import { ClientSpec, getPerforceUsername, getRootDirectoryForBranch, initializePerforce, PerforceContext, StreamSpecs } from '../common/perforce';
 import { BuildVersion, VersionReader } from '../common/version';
 import { CertFiles } from '../common/webserver';
 import { addBranchGraph, Graph, GraphAPI } from '../new/graph';
