NodeBot will now correctly be considered active while reconsidering shelf requests
This was broken in CL# 17343508

diff --git a/Engine/Extras/RoboMerge/v3/src/robo/graphbot.ts b/Engine/Extras/RoboMerge/v3/src/robo/graphbot.ts
--- a/Engine/Extras/RoboMerge/v3/src/robo/graphbot.ts
+++ b/Engine/Extras/RoboMerge/v3/src/robo/graphbot.ts
@@ -216,7 +216,9 @@ export class GraphBot implements GraphInterface, BotEventHandler {
 					break
 				}
 				try {
+					node.isActive = true
 					await node.processQueuedChange(request)
+					node.isActive = false
 				}
 				catch(err) {
 					return [node, err]
@@ -264,7 +266,7 @@ export class GraphBot implements GraphInterface, BotEventHandler {
 
 		for (const bot of this.botlist) {
 			if (!bot.isRunning) {
-				this.botLogger.debug(`Starting bot ${bot.fullNameForLogging}`)
+				this.botLogger.info(`Starting bot ${bot.fullNameForLogging}`)
 				bot.start()
 			}
 		}
@@ -274,6 +276,7 @@ export class GraphBot implements GraphInterface, BotEventHandler {
 
 			for (const bot of this.botlist) {
 				bot.isActive = true
+				this.botLogger.info(`Bot ${bot.fullNameForLogging} is active`)
 				let ticked = false
 				try {
 					// crashMe API support - simulate a bot crashing and stopping the GraphBot instance
@@ -289,6 +292,7 @@ export class GraphBot implements GraphInterface, BotEventHandler {
 					this.handleNodebotError(bot, err)
 					return
 				}
+				this.botLogger.info(`Bot ${bot.fullNameForLogging} is no longer active`)
 				bot.isActive = false
 
 				if (ticked) {
