Remove debugging logging

diff --git a/Engine/Extras/RoboMerge/v3/src/robo/graphbot.ts b/Engine/Extras/RoboMerge/v3/src/robo/graphbot.ts
--- a/Engine/Extras/RoboMerge/v3/src/robo/graphbot.ts
+++ b/Engine/Extras/RoboMerge/v3/src/robo/graphbot.ts
@@ -266,7 +266,7 @@ export class GraphBot implements GraphInterface, BotEventHandler {
 
 		for (const bot of this.botlist) {
 			if (!bot.isRunning) {
-				this.botLogger.info(`Starting bot ${bot.fullNameForLogging}`)
+				this.botLogger.debug(`Starting bot ${bot.fullNameForLogging}`)
 				bot.start()
 			}
 		}
@@ -276,7 +276,6 @@ export class GraphBot implements GraphInterface, BotEventHandler {
 
 			for (const bot of this.botlist) {
 				bot.isActive = true
-				this.botLogger.info(`Bot ${bot.fullNameForLogging} is active`)
 				let ticked = false
 				try {
 					// crashMe API support - simulate a bot crashing and stopping the GraphBot instance
@@ -292,7 +291,6 @@ export class GraphBot implements GraphInterface, BotEventHandler {
 					this.handleNodebotError(bot, err)
 					return
 				}
-				this.botLogger.info(`Bot ${bot.fullNameForLogging} is no longer active`)
 				bot.isActive = false
 
 				if (ticked) {
