Fix consideration of persistence backup frequency to be minutes, not seconds

diff --git a/Engine/Extras/RoboMerge/v3/src/robo/settings.ts b/Engine/Extras/RoboMerge/v3/src/robo/settings.ts
--- a/Engine/Extras/RoboMerge/v3/src/robo/settings.ts
+++ b/Engine/Extras/RoboMerge/v3/src/robo/settings.ts
@@ -457,7 +457,7 @@ export class Settings {
 						}
 					})
 				}
-				else if ((Date.now() - this.lastPersistBackupTime.getTime()) / 1000 >= args.persistenceBackupFrequency) {
+				else if ((Date.now() - this.lastPersistBackupTime.getTime()) / 60000 >= args.persistenceBackupFrequency) {
 					this.lastPersistBackupTime = new Date(Date.now())
 
 					this.p4.sync(this.p4workspace, persistentFile, {opts: ['-k']})
