Allow reconsiders to choose edge server workspaces

diff --git a/Engine/Extras/RoboMerge/v3/src/robo/edgebot.ts b/Engine/Extras/RoboMerge/v3/src/robo/edgebot.ts
--- a/Engine/Extras/RoboMerge/v3/src/robo/edgebot.ts
+++ b/Engine/Extras/RoboMerge/v3/src/robo/edgebot.ts
@@ -865,9 +865,7 @@ class EdgeBotImpl extends PerforceStatefulBot {
 		// Find a suitable workspace from one of the owner's workspaces
 		else if (!forApproval) {
 			// use p4.find_workspaces to find a workspace (owned by the user) for this change if this is a stream branch
-			// not supporting edge servers yet - really need to send a separate shelf request instead
-			// so specify to only look for workspaces on the commit server
-			const workspaces: ClientSpec[] = await p4util.getWorkspacesForUser(this.p4, owner, 'commit')
+			const workspaces: ClientSpec[] = await p4util.getWorkspacesForUser(this.p4, owner)
 
 			if (workspaces.length > 0) {
 				// default to the first workspace
