When reconsidering do the operation on the edgeserver that the shelf will eventually need to be placed on if there is a conflict

diff --git a/Engine/Extras/RoboMerge/v3/src/robo/edgebot.ts b/Engine/Extras/RoboMerge/v3/src/robo/edgebot.ts
--- a/Engine/Extras/RoboMerge/v3/src/robo/edgebot.ts
+++ b/Engine/Extras/RoboMerge/v3/src/robo/edgebot.ts
@@ -22,7 +22,7 @@ import { PauseState } from "./state-interfaces";
 import { BlockagePauseInfo, BlockagePauseInfoMinimal, EdgeStatusFields } from "./status-types";
 import { getIntegrationOwner } from "./targets";
 
-function matchPrefix(a: string, b: string) {
+export function matchPrefix(a: string, b: string) {
 	const len = Math.min(a.length, b.length)
 	for (let i = 0; i < len; ++i) {
 		if (a.charAt(i) !== b.charAt(i)) {
@@ -899,7 +899,7 @@ class EdgeBotImpl extends PerforceStatefulBot {
 			opts.newWorkspace = targetWorkspace
 		}
 		else if (!forApproval) {
-			this.edgeBotLogger.warn(`Unable to find appropriate workspace for ${owner}` + (branch_stream || this.targetBranch.name))
+			this.edgeBotLogger.warn(`Unable to find appropriate workspace for ${owner} ` + (branch_stream || this.targetBranch.name))
 		}
 
 		// edit the owner to the author so they can resolve and submit themselves
diff --git a/Engine/Extras/RoboMerge/v3/src/robo/nodebot.ts b/Engine/Extras/RoboMerge/v3/src/robo/nodebot.ts
--- a/Engine/Extras/RoboMerge/v3/src/robo/nodebot.ts
+++ b/Engine/Extras/RoboMerge/v3/src/robo/nodebot.ts
@@ -1,17 +1,18 @@
 // Copyright Epic Games, Inc. All Rights Reserved.
 
 import * as Sentry from '@sentry/node';
+import * as p4util from '../common/p4util';
 import { _nextTick } from '../common/helper';
 import { ContextualLogger } from '../common/logger';
 import { Mailer, MailParams, Recipients } from '../common/mailer';
-import { Change, ConflictedResolveNFile, PerforceContext, RoboWorkspace, coercePerforceWorkspace } from '../common/perforce';
+import { Change, ClientSpec, ConflictedResolveNFile, PerforceContext, RoboWorkspace, coercePerforceWorkspace } from '../common/perforce';
 import { IPCControls, NodeBotInterface, QueuedChange, ReconsiderArgs } from './bot-interfaces';
 import { ApprovalOptions, BotConfig, EdgeOptions, NodeOptions } from './branchdefs'
 import { BlockagePauseInfo, BranchStatus } from './status-types';
 import { AlreadyIntegrated, Blockage, Branch, BranchArg, BranchGraphInterface, ChangeInfo, EndIntegratingToGateEvent, Failure } from './branch-interfaces';
 import { MergeAction, OperationResult, PendingChange, resolveBranchArg, StompedRevision, StompVerification, StompVerificationFile }  from './branch-interfaces';
 import { Conflicts } from './conflicts';
-import { EdgeBot, EdgeMergeResults } from './edgebot';
+import { EdgeBot, EdgeMergeResults, matchPrefix } from './edgebot';
 import { BotEventTriggers } from './events';
 import { makeClLink, SlackMessages } from './notifications';
 import { PerforceStatefulBot } from './perforce-stateful-bot';
@@ -1463,12 +1464,42 @@ export class NodeBot extends PerforceStatefulBot implements NodeBotInterface {
 		}
 
 		// should also do this for #manual changes (set up some testing around those first)
-		if (change.forceCreateAShelf && optWorkspaceOverride) {
-			// see if we need to talk to an edge server to create the shelf
-			const edgeServer = await this.p4.getWorkspaceEdgeServer(
-				coercePerforceWorkspace(optWorkspaceOverride)!.name)
-			if (edgeServer) {
-				result.info.edgeServerToHostShelf = edgeServer
+		if (change.forceCreateAShelf || change.isUserRequest) {
+			if (!optWorkspaceOverride && (result.info.targets || []).length == 1) {
+				const workspaces: ClientSpec[] = await p4util.getWorkspacesForUser(this.p4, result.info.owner!)
+
+				if (workspaces.length > 0) {
+					// default to the first workspace
+					let targetWorkspace = workspaces[0].client
+	
+					// if this is a stream branch, do some better match-up
+					if (result.info.targets![0].branch.stream) {
+						const branch_stream = result.info.targets![0].branch.stream.toLowerCase()
+						// find the stream with the closest match
+						let target_match = 0
+						for (let def of workspaces) {
+							let stream = def.Stream
+							if (stream) {
+								const matchlen = matchPrefix(stream.toLowerCase(), branch_stream)
+								if (matchlen > target_match) {
+									target_match = matchlen
+									targetWorkspace = def.client
+								}
+							}
+						}
+					}
+					this.nodeBotLogger.info(`Chose workspace ${targetWorkspace}`)
+					result.info.targetWorkspaceForShelf = targetWorkspace
+					optWorkspaceOverride = targetWorkspace
+				}
+			}
+			if (optWorkspaceOverride) {
+				// see if we need to talk to an edge server to create the shelf
+				const edgeServer = await this.p4.getWorkspaceEdgeServer(
+					coercePerforceWorkspace(optWorkspaceOverride)!.name)
+				if (edgeServer) {
+					result.info.edgeServerToHostShelf = edgeServer
+				}
 			}
 		}
 
