Disable reconsider edge server code path temporarily

diff --git a/Engine/Extras/RoboMerge/v3/src/robo/nodebot.ts b/Engine/Extras/RoboMerge/v3/src/robo/nodebot.ts
--- a/Engine/Extras/RoboMerge/v3/src/robo/nodebot.ts
+++ b/Engine/Extras/RoboMerge/v3/src/robo/nodebot.ts
@@ -1464,7 +1464,8 @@ export class NodeBot extends PerforceStatefulBot implements NodeBotInterface {
 		}
 
 		// should also do this for #manual changes (set up some testing around those first)
-		if (change.forceCreateAShelf || change.isUserRequest) {
+		if (change.forceCreateAShelf && optWorkspaceOverride) {
+		//if (change.forceCreateAShelf || change.isUserRequest) {
 			if (!optWorkspaceOverride && (result.info.targets || []).length == 1) {
 				const workspaces: ClientSpec[] = await p4util.getWorkspacesForUser(this.p4, result.info.owner!)
 
