Fix track change failure for initial submits in a subpath stream

diff --git a/Engine/Extras/RoboMerge/v3/src/robo/ipc.ts b/Engine/Extras/RoboMerge/v3/src/robo/ipc.ts
--- a/Engine/Extras/RoboMerge/v3/src/robo/ipc.ts
+++ b/Engine/Extras/RoboMerge/v3/src/robo/ipc.ts
@@ -237,7 +237,7 @@ export class IPC {
 			let desc = await this.robo.p4.describe(cl, 1)
 			const mergeMethod = cl != data.originalCL ? getMergeMethod(desc.description) : 'initialSubmit'
 			const clNode = getNode(desc)
-			if (clNode || mergeMethod in RobomergeMethodStrings) {
+			if (clNode || (RobomergeMethodStrings as readonly string[]).includes(mergeMethod)) {
 				changes.set(cl, {desc, node: clNode, sourceCL: opts ? opts.sourceCL : null, destCLs: (opts && opts.lastCL ? [opts.lastCL] : []) })
 				return true
 			}
@@ -338,7 +338,7 @@ export class IPC {
 					}
 					break
 				} else if (hasAutomergeTarget && changeToConsider.desc.entries.length == 1 && 
-							(mergeMethod in RobomergeMethodStrings)) {
+							(RobomergeMethodStrings as readonly string[]).includes(mergeMethod)) {
 					// If we only have 1 entry and we didn't get integration info off of it
 					// and the graph suggests we are expecting there could be other changes
 					// get the full describe results
