Fix reconsider submits occurring on commit server when the shelf is on an edgeserver

diff --git a/Engine/Extras/RoboMerge/v3/src/common/perforce.ts b/Engine/Extras/RoboMerge/v3/src/common/perforce.ts
--- a/Engine/Extras/RoboMerge/v3/src/common/perforce.ts
+++ b/Engine/Extras/RoboMerge/v3/src/common/perforce.ts
@@ -1027,11 +1027,11 @@ export class PerforceContext {
 
 	// submit a CL
 	// output format is final CL number or false if changes need more resolution
-	async submit(roboWorkspace: RoboWorkspace, changelist: number): Promise<number | string> {
+	async submit(roboWorkspace: RoboWorkspace, changelist: number, edgeServerAddress?: string): Promise<number | string> {
 		const workspace = coercePerforceWorkspace(roboWorkspace)
 		let rawOutput: string
 		try {
-			rawOutput = await this._execP4(workspace, ['-ztag', 'submit', '-f', 'submitunchanged', '-c', changelist.toString()])
+			rawOutput = await this._execP4(workspace, ['-ztag', 'submit', '-f', 'submitunchanged', '-c', changelist.toString()], {edgeServerAddress})
 		}
 		catch ([errArg, output]) {
 			const err = errArg.toString().trim()
diff --git a/Engine/Extras/RoboMerge/v3/src/robo/edgebot.ts b/Engine/Extras/RoboMerge/v3/src/robo/edgebot.ts
--- a/Engine/Extras/RoboMerge/v3/src/robo/edgebot.ts
+++ b/Engine/Extras/RoboMerge/v3/src/robo/edgebot.ts
@@ -728,7 +728,11 @@ class EdgeBotImpl extends PerforceStatefulBot {
 
 		// try to submit
 		this._log_action(`Submitting CL ${changenum} by ${info.author}`)
-		const result = await this.p4.submit(this.targetBranch.workspace, changenum)
+		const result = await this.p4.submit(
+			pending.change.targetWorkspaceOverride || this.targetBranch.workspace, 
+			changenum,
+			pending.change.edgeServerToHostShelf && pending.change.edgeServerToHostShelf.address)
+			
 		if (typeof result === "string") {
 			// an error occurred while submitting
 			return new EdgeIntegrationDetails('error', result);
