Centralize workspace selection code
Bias workspace selection to most recently accessed

diff --git a/Engine/Extras/RoboMerge/v3/src/common/p4util.ts b/Engine/Extras/RoboMerge/v3/src/common/p4util.ts
--- a/Engine/Extras/RoboMerge/v3/src/common/p4util.ts
+++ b/Engine/Extras/RoboMerge/v3/src/common/p4util.ts
@@ -1,7 +1,7 @@
 // Copyright Epic Games, Inc. All Rights Reserved.
 
 import { ContextualLogger } from './logger';
-import { Change, coercePerforceWorkspace, OpenedFileRecord, PerforceContext, RoboWorkspace } from './perforce';
+import { Change, ClientSpec, coercePerforceWorkspace, OpenedFileRecord, PerforceContext, RoboWorkspace } from './perforce';
 
 const USER_WORKSPACE_EXCLUDE_PATTERNS: (RegExp | string)[] = [
 	'horde-p4bridge-',
@@ -223,8 +223,47 @@ export async function cleanWorkspaces(p4: PerforceContext, workspaces: [string,
 	await Promise.all(workspaces.map(([name, root]) => p4.sync(name, root + '#0', {edgeServerAddress})))
 }
 
-export async function getWorkspacesForUser(p4: PerforceContext, user: string, edgeServerAddress?: string) {
-	return (await p4.find_workspaces(user, {edgeServerAddress}))
+export async function getWorkspacesForUser(p4: PerforceContext, user: string) {
+	return (await p4.find_workspaces(user))
 		.filter(ws => !USER_WORKSPACE_EXCLUDE_PATTERNS.some(entry => ws.client.match(entry)))
 }
 
+function matchPrefix(a: string, b: string) {
+	const len = Math.min(a.length, b.length)
+	for (let i = 0; i < len; ++i) {
+		if (a.charAt(i) !== b.charAt(i)) {
+			return i
+		}
+	}
+	return len
+}
+
+export async function chooseBestWorkspaceForUser(p4: PerforceContext, user: string, stream?: string) {
+	const workspaces: ClientSpec[] = await getWorkspacesForUser(p4,user)
+
+	if (workspaces.length > 0) {
+		workspaces.sort((a,b) => b.Access - a.Access)
+		// default to the first workspace
+		let targetWorkspace = workspaces[0]
+
+		// if this is a stream branch, do some better match-up
+		if (stream) {
+			const branch_stream = stream.toLowerCase()
+			// find the stream with the closest match
+			let target_match = 0
+			for (let def of workspaces) {
+				let stream = def.Stream
+				if (stream) {
+					const matchlen = matchPrefix(stream.toLowerCase(), branch_stream)
+					if (matchlen > target_match) {
+						target_match = matchlen
+						targetWorkspace = def
+					}
+				}
+			}
+		}
+		return targetWorkspace.client
+	}
+
+	return undefined
+}
\ No newline at end of file
diff --git a/Engine/Extras/RoboMerge/v3/src/common/perforce.ts b/Engine/Extras/RoboMerge/v3/src/common/perforce.ts
--- a/Engine/Extras/RoboMerge/v3/src/common/perforce.ts
+++ b/Engine/Extras/RoboMerge/v3/src/common/perforce.ts
@@ -257,6 +257,7 @@ export type RoboWorkspace = Workspace | string | null;
 
 export interface ClientSpec {
 	client: string
+	Access: number
 	Stream?: string
 	IsUnloaded?: boolean
 }
diff --git a/Engine/Extras/RoboMerge/v3/src/robo/edgebot.ts b/Engine/Extras/RoboMerge/v3/src/robo/edgebot.ts
--- a/Engine/Extras/RoboMerge/v3/src/robo/edgebot.ts
+++ b/Engine/Extras/RoboMerge/v3/src/robo/edgebot.ts
@@ -5,7 +5,7 @@ import * as p4util from '../common/p4util';
 
 import { ContextualLogger } from "../common/logger";
 import { Recipients } from "../common/mailer";
-import { Change, ClientSpec, coercePerforceWorkspace, ConflictedResolveNFile, EditChangeOpts, EXCLUSIVE_CHECKOUT_REGEX, getRootDirectoryForBranch, IntegrationSource, IntegrationTarget, isExecP4Error, OpenedFileRecord, PerforceContext } from "../common/perforce";
+import { Change, coercePerforceWorkspace, ConflictedResolveNFile, EditChangeOpts, EXCLUSIVE_CHECKOUT_REGEX, getRootDirectoryForBranch, IntegrationSource, IntegrationTarget, isExecP4Error, OpenedFileRecord, PerforceContext } from "../common/perforce";
 import { VersionReader } from "../common/version";
 import { EdgeBotInterface, IPCControls, ReconsiderArgs } from "./bot-interfaces";
 import { AlreadyIntegrated, Branch, ChangeInfo, ConflictingFile, Failure, MergeAction, PendingChange } from "./branch-interfaces";
@@ -22,16 +22,6 @@ import { PauseState } from "./state-interfaces";
 import { BlockagePauseInfo, BlockagePauseInfoMinimal, EdgeStatusFields } from "./status-types";
 import { getIntegrationOwner } from "./targets";
 
-export function matchPrefix(a: string, b: string) {
-	const len = Math.min(a.length, b.length)
-	for (let i = 0; i < len; ++i) {
-		if (a.charAt(i) !== b.charAt(i)) {
-			return i
-		}
-	}
-	return len
-}
-
 const FAILED_CHANGELIST_PAUSE_TIMEOUT_SECONDS = 15 * 60
 const MAX_INTEGRATION_ERRORS_TO_ANALYZE = 5
 const DEPOT_FILE_REGEX = /^(.*[\\\/])(.*)/
@@ -860,7 +850,7 @@ class EdgeBotImpl extends PerforceStatefulBot {
 		await this.p4.revert(destRoboWorkspace, changenum, [], edgeServerAddress)
 
 		// figure out what workspace to put it in
-		const branch_stream = this.targetBranch.stream ? this.targetBranch.stream.toLowerCase() : null
+		const branch_stream = this.targetBranch.stream ? this.targetBranch.stream.toLowerCase() : undefined
 		let targetWorkspace: string | undefined = undefined
 		// Check for specified workspace from createShelf operation
 		if (pending.change.targetWorkspaceForShelf) {
@@ -869,30 +859,9 @@ class EdgeBotImpl extends PerforceStatefulBot {
 		// Find a suitable workspace from one of the owner's workspaces
 		else if (!forApproval) {
 			// use p4.find_workspaces to find a workspace (owned by the user) for this change if this is a stream branch
-			const workspaces: ClientSpec[] = await p4util.getWorkspacesForUser(this.p4, owner)
-
-			if (workspaces.length > 0) {
-				// default to the first workspace
-				targetWorkspace = workspaces[0].client
-
-				// if this is a stream branch, do some better match-up
-				if (branch_stream) {
-					// find the stream with the closest match
-					let target_match = 0
-					for (let def of workspaces) {
-						let stream = def.Stream
-						if (stream) {
-							let matchlen = matchPrefix(stream.toLowerCase(), branch_stream)
-							if (matchlen > target_match) {
-								target_match = matchlen
-								targetWorkspace = def.client
-							}
-						}
-					}
-				}
-				this.edgeBotLogger.info(`Chose workspace ${targetWorkspace} (stream? ${branch_stream ? 'yes' : 'no'})`)
-				pending.change.targetWorkspaceForShelf = targetWorkspace
-			}
+			targetWorkspace = await p4util.chooseBestWorkspaceForUser(this.p4, owner, branch_stream)
+			pending.change.targetWorkspaceForShelf = targetWorkspace
+			this.edgeBotLogger.info(`Chose workspace ${targetWorkspace}`)
 		}
 
 		// log if we couldn't find a workspace
diff --git a/Engine/Extras/RoboMerge/v3/src/robo/nodebot.ts b/Engine/Extras/RoboMerge/v3/src/robo/nodebot.ts
--- a/Engine/Extras/RoboMerge/v3/src/robo/nodebot.ts
+++ b/Engine/Extras/RoboMerge/v3/src/robo/nodebot.ts
@@ -5,14 +5,14 @@ import * as p4util from '../common/p4util';
 import { _nextTick } from '../common/helper';
 import { ContextualLogger } from '../common/logger';
 import { Mailer, MailParams, Recipients } from '../common/mailer';
-import { Change, ClientSpec, ConflictedResolveNFile, PerforceContext, RoboWorkspace, coercePerforceWorkspace } from '../common/perforce';
+import { Change, ConflictedResolveNFile, PerforceContext, RoboWorkspace, coercePerforceWorkspace } from '../common/perforce';
 import { IPCControls, NodeBotInterface, QueuedChange, ReconsiderArgs } from './bot-interfaces';
 import { ApprovalOptions, BotConfig, EdgeOptions, NodeOptions } from './branchdefs'
 import { BlockagePauseInfo, BranchStatus } from './status-types';
 import { AlreadyIntegrated, Blockage, Branch, BranchArg, BranchGraphInterface, ChangeInfo, EndIntegratingToGateEvent, Failure } from './branch-interfaces';
 import { MergeAction, OperationResult, PendingChange, resolveBranchArg, StompedRevision, StompVerification, StompVerificationFile }  from './branch-interfaces';
 import { Conflicts } from './conflicts';
-import { EdgeBot, EdgeMergeResults, matchPrefix } from './edgebot';
+import { EdgeBot, EdgeMergeResults } from './edgebot';
 import { BotEventTriggers } from './events';
 import { makeClLink, SlackMessages } from './notifications';
 import { PerforceStatefulBot } from './perforce-stateful-bot';
@@ -1466,32 +1466,9 @@ export class NodeBot extends PerforceStatefulBot implements NodeBotInterface {
 		// should also do this for #manual changes (set up some testing around those first)
 		if (change.forceCreateAShelf || (change.isUserRequest && !change.forceStompChanges)) {
 			if (!optWorkspaceOverride && (result.info.targets || []).length == 1) {
-				const workspaces: ClientSpec[] = await p4util.getWorkspacesForUser(this.p4, result.info.owner!)
-
-				if (workspaces.length > 0) {
-					// default to the first workspace
-					let targetWorkspace = workspaces[0].client
-	
-					// if this is a stream branch, do some better match-up
-					if (result.info.targets![0].branch.stream) {
-						const branch_stream = result.info.targets![0].branch.stream.toLowerCase()
-						// find the stream with the closest match
-						let target_match = 0
-						for (let def of workspaces) {
-							let stream = def.Stream
-							if (stream) {
-								const matchlen = matchPrefix(stream.toLowerCase(), branch_stream)
-								if (matchlen > target_match) {
-									target_match = matchlen
-									targetWorkspace = def.client
-								}
-							}
-						}
-					}
-					this.nodeBotLogger.info(`Chose workspace ${targetWorkspace}`)
-					result.info.targetWorkspaceForShelf = targetWorkspace
-					optWorkspaceOverride = targetWorkspace
-				}
+				result.info.targetWorkspaceForShelf = await p4util.chooseBestWorkspaceForUser(this.p4, result.info.owner!, result.info.targets![0].branch.stream)
+				optWorkspaceOverride = result.info.targetWorkspaceForShelf
+				this.nodeBotLogger.info(`Chose workspace ${optWorkspaceOverride}`)
 			}
 			if (optWorkspaceOverride) {
 				// see if we need to talk to an edge server to create the shelf
