Bring some set up files referenced in installation doc in to public visible repo
Clean out ldap config that doesn't need to be in the public version
#rnx

diff --git a/Engine/Extras/RoboMerge/v3/config/ldap.cfg.json b/Engine/Extras/RoboMerge/v3/config/ldap.cfg.json
--- a/Engine/Extras/RoboMerge/v3/config/ldap.cfg.json
+++ b/Engine/Extras/RoboMerge/v3/config/ldap.cfg.json
@@ -18,10 +18,6 @@
 		{
 			"group": "Robomerge_FTE_Access",
 			"tags": ["fte"]
-		},
-		{
-			"group": "Robomerge_Contractor_Twinmotion_Access",
-			"tags": ["twinmotion"]
 		}
 	]
 }
\ No newline at end of file
diff --git a/Engine/Extras/RoboMerge/v3/local-docker-build.bat b/Engine/Extras/RoboMerge/v3/local-docker-build.bat
new file mode 100755
--- /dev/null
+++ b/Engine/Extras/RoboMerge/v3/local-docker-build.bat
@@ -0,0 +1,31 @@
+@ECHO OFF
+
+set EPIC_BUILD_ROLE_NAME=robomerge-ts-service-testing
+set DOCKERFILE=Dockerfile
+
+set DOCKER_REGISTRY_DOMAIN=
+set DOCKER_REGISTRY_NAMESPACE=
+set DOCKER_IMAGE_NAME=%DOCKER_REGISTRY_DOMAIN%/%DOCKER_REGISTRY_NAMESPACE%/%EPIC_BUILD_ROLE_NAME%
+
+docker version >NUL || echo ERROR: Could not connect to Docker daemon. && goto :end -1
+
+echo Building %DOCKER_IMAGE_NAME% ...
+docker build --file %DOCKERFILE% --pull --rm --tag %DOCKER_IMAGE_NAME%:latest . || goto :end -1
+
+echo.
+echo.
+echo Waiting for user to test and will push...
+pause
+
+echo.
+echo.
+echo Pushing %DOCKER_IMAGE_NAME%:latest to %DOCKER_REGISTRY_DOMAIN%/%DOCKER_REGISTRY_NAMESPACE%
+docker push %DOCKER_IMAGE_NAME%:latest
+
+echo Success!
+goto :end 0
+
+
+:end
+pause
+exit /b %~1
diff --git a/Engine/Extras/RoboMerge/v3/local-docker-run.bat b/Engine/Extras/RoboMerge/v3/local-docker-run.bat
new file mode 100755
--- /dev/null
+++ b/Engine/Extras/RoboMerge/v3/local-docker-run.bat
@@ -0,0 +1,48 @@
+@ECHO OFF
+set EPIC_BUILD_ROLE_NAME=robomerge-ts-service-testing
+
+set DOCKER_REGISTRY_DOMAIN=
+set DOCKER_REGISTRY_NAMESPACE=
+set DOCKER_IMAGE_NAME=%DOCKER_REGISTRY_DOMAIN%/%DOCKER_REGISTRY_NAMESPACE%/%EPIC_BUILD_ROLE_NAME%
+set DOCKER_VERSION=latest
+
+set ROBO_LOG_LEVEL=info
+
+REM Set P4PORT to IP of Perforce master server to bypass any DNS issues
+set P4PORT=perforce:1666
+set P4PASSWD=
+set BOTS=robomergeQA1
+set ROBO_EXTERNAL_URL=https://localhost:4433
+set PORTS_ARGS=-p 4433:4433
+
+REM set NODE_ENV=
+REM set EPIC_ENV=
+REM set EPIC_DEPLOYMENT=
+REM set SENTRY_DSN=
+
+@ECHO ON
+
+REM docker pull --platform linux %DOCKER_IMAGE_NAME%:%DOCKER_VERSION%
+docker pull %DOCKER_IMAGE_NAME%:%DOCKER_VERSION%
+
+docker stop %EPIC_BUILD_ROLE_NAME%>NUL
+docker rm %EPIC_BUILD_ROLE_NAME%>NUL
+
+@REM -v robosettings:/root/.robomerge \
+
+docker run -d --name %EPIC_BUILD_ROLE_NAME% ^
+    -e "P4PASSWD=%P4PASSWD%" ^
+    -e "P4PORT=%P4PORT%" ^
+    -e "BOTNAME=%BOTS%" ^
+    -e "ROBO_EXTERNAL_URL=%ROBO_EXTERNAL_URL%" ^
+    -e "NODE_ENV=%NODE_ENV%" ^
+    -e "EPIC_ENV=%EPIC_ENV%" ^
+    -e "EPIC_DEPLOYMENT=%EPIC_DEPLOYMENT%" ^
+    -e "SENTRY_DSN=%SENTRY_DSN%" ^
+    -e "ROBO_LOG_LEVEL=%ROBO_LOG_LEVEL%" ^
+    -v "%CD%\vault:/vault:ro" ^
+    -v robomergesavedata:/root/.robomerge ^
+    %PORTS_ARGS% ^
+    %DOCKER_IMAGE_NAME%:%DOCKER_VERSION% 
+
+pause
\ No newline at end of file
diff --git a/Engine/Extras/RoboMerge/v3/local-docker-run.sh b/Engine/Extras/RoboMerge/v3/local-docker-run.sh
new file mode 100644
--- /dev/null
+++ b/Engine/Extras/RoboMerge/v3/local-docker-run.sh
@@ -0,0 +1,43 @@
+set +x
+export EPIC_BUILD_ROLE_NAME=robomerge-ts-service-testing
+
+export DOCKER_REGISTRY_DOMAIN=
+export DOCKER_REGISTRY_NAMESPACE=
+export DOCKER_IMAGE_NAME=$DOCKER_REGISTRY_DOMAIN/$DOCKER_REGISTRY_NAMESPACE/$EPIC_BUILD_ROLE_NAME
+export DOCKER_VERSION=latest
+
+export ROBO_LOG_LEVEL=info
+
+# export NODE_ENV=
+# export EPIC_ENV=
+# export EPIC_DEPLOYMENT=
+# export SENTRY_DSN=
+
+# Set P4PORT to IP of Perforce master server to bypass any DNS issues
+export P4PORT=perforce:1666
+export P4PASSWD=
+export BOTS=robomergeQA1
+export ROBO_EXTERNAL_URL=http://localhost:8877
+
+set -x
+
+docker pull $DOCKER_IMAGE_NAME:$DOCKER_VERSION
+
+docker stop $EPIC_BUILD_ROLE_NAME > /dev/null
+docker rm $EPIC_BUILD_ROLE_NAME  > /dev/null
+
+docker run -d --name $EPIC_BUILD_ROLE_NAME \
+    -e "P4PASSWD=$P4PASSWD" \
+    -e "P4PORT=$P4PORT" \
+    -e "BOTNAME=$BOTS" \
+    -e "ROBO_EXTERNAL_URL=$ROBO_EXTERNAL_URL" \
+    -e "NODE_ENV=$NODE_ENV" \
+    -e "EPIC_ENV=$EPIC_ENV" \
+    -e "EPIC_DEPLOYMENT=$EPIC_DEPLOYMENT" \
+    -e "SENTRY_DSN=$SENTRY_DSN" \
+    -e "ROBO_LOG_LEVEL=$ROBO_LOG_LEVEL"
+    -p 8877:8877 \
+    -p 1666:1666 \
+    -v robomergesavedata:/root/.robomerge \
+    $DOCKER_IMAGE_NAME:$DOCKER_VERSION \
+    node dist/robo/watchdog.js 
\ No newline at end of file
