Make sure a workspace is loaded before trying to use it to integrate

diff --git a/Engine/Extras/RoboMerge/v3/src/robo/edgebot.ts b/Engine/Extras/RoboMerge/v3/src/robo/edgebot.ts
--- a/Engine/Extras/RoboMerge/v3/src/robo/edgebot.ts
+++ b/Engine/Extras/RoboMerge/v3/src/robo/edgebot.ts
@@ -483,7 +483,11 @@ class EdgeBotImpl extends PerforceStatefulBot {
 
 			// do we already have a workspace? (_initBranchWorkspacesForAllBots, _getExistingWorkspaces)
 			const existingWorkspaceInfos = await this.p4.find_workspaces(undefined, {edgeServerAddress: edgeServer.address, includeUnloaded: true})
-			if (existingWorkspaceInfos.map(ws => ws.client).indexOf(info.targetWorkspaceOverride) >= 0) {
+			const existingWorkspaceIndex = existingWorkspaceInfos.findIndex((ws) => ws.client == info.targetWorkspaceOverride)
+			if (existingWorkspaceIndex >= 0) {
+				if (existingWorkspaceInfos[existingWorkspaceIndex].IsUnloaded) {
+					await this.p4.reloadWorkspace(info.targetWorkspaceOverride)
+				}
 				await p4util.cleanWorkspaces(this.p4, [[info.targetWorkspaceOverride, target.branch.rootPath]], edgeServer.address)
 			}
 			else {
