Disable unnecessary logging

diff --git a/Engine/Extras/RoboMerge/v3/web/src/app.ts b/Engine/Extras/RoboMerge/v3/web/src/app.ts
--- a/Engine/Extras/RoboMerge/v3/web/src/app.ts
+++ b/Engine/Extras/RoboMerge/v3/web/src/app.ts
@@ -19,12 +19,12 @@ function refreshAll() {
 }
 
 function pauseAutoRefresh() {
-	log('Pausing auto refresh of branch list')
+	//log('Pausing auto refresh of branch list')
 	pauseRefresh = true
 }
 
 function resumeAutoRefresh() {
-	log('Resuming auto refresh of branch list')
+	//log('Resuming auto refresh of branch list')
 	pauseRefresh = false
 }
 
