Fix track changes failing if the first file in the changelist is a move/delete record

diff --git a/Engine/Extras/RoboMerge/v3/src/robo/ipc.ts b/Engine/Extras/RoboMerge/v3/src/robo/ipc.ts
--- a/Engine/Extras/RoboMerge/v3/src/robo/ipc.ts
+++ b/Engine/Extras/RoboMerge/v3/src/robo/ipc.ts
@@ -320,24 +320,28 @@ export class IPC {
 
 			for (let i=0; i < changeToConsider.desc.entries.length; i++) {
 				const entry = changeToConsider.desc.entries[i]
-				const integrated = await this.robo.p4.integrated(null, entry.depotFile, {intoOnly: true, startCL: clToConsider})
-				if (integrated.length > 0) {
-					for (let integ of integrated) {
-						const startToRev = integ.startToRev == "#none" ? 0 : parseInt(integ.startToRev.slice(1))
-						const endToRev = parseInt(integ.endToRev.slice(1))
-						if (entry.rev <= endToRev && entry.rev > startToRev)
-						{
-							const destChange = changes.get(integ.change)
-							if (destChange) {
-								destChange.sourceCL = clToConsider
-							} else {
-								changeToConsider.destCLs.push(integ.change)
-								await gatherCLInfo(integ.change, {sourceCL: clToConsider})
+				// integrated for move/delete will point at the paired move/add not where it was merged to in another stream, so not useful to evaluate it
+				if (entry.action != 'move/delete') {
+					const integrated = await this.robo.p4.integrated(null, entry.depotFile, {intoOnly: true, startCL: clToConsider})
+					if (integrated.length > 0) {
+						for (let integ of integrated) {
+							const startToRev = integ.startToRev == "#none" ? 0 : parseInt(integ.startToRev.slice(1))
+							const endToRev = parseInt(integ.endToRev.slice(1))
+							if (entry.rev <= endToRev && entry.rev > startToRev)
+							{
+								const destChange = changes.get(integ.change)
+								if (destChange) {
+									destChange.sourceCL = clToConsider
+								} else {
+									changeToConsider.destCLs.push(integ.change)
+									await gatherCLInfo(integ.change, {sourceCL: clToConsider})
+								}
 							}
 						}
+						break
 					}
-					break
-				} else if (hasAutomergeTarget && changeToConsider.desc.entries.length == 1 && 
+				}
+				if (hasAutomergeTarget && changeToConsider.desc.entries.length == 1 && 
 							(RobomergeMethodStrings as readonly string[]).includes(mergeMethod)) {
 					// If we only have 1 entry and we didn't get integration info off of it
 					// and the graph suggests we are expecting there could be other changes
