Reload workspaces on the correct edgeserver

diff --git a/Engine/Extras/RoboMerge/v3/src/common/perforce.ts b/Engine/Extras/RoboMerge/v3/src/common/perforce.ts
--- a/Engine/Extras/RoboMerge/v3/src/common/perforce.ts
+++ b/Engine/Extras/RoboMerge/v3/src/common/perforce.ts
@@ -645,8 +645,8 @@ export class PerforceContext {
 		return null
 	}
 
-	reloadWorkspace(workspaceName: string) {
-		return this._execP4Ztag(null, ['reload', '-c', workspaceName]);
+	reloadWorkspace(workspaceName: string, edgeServerAddress?: string) {
+		return this._execP4Ztag(null, ['reload', '-c', workspaceName], {edgeServerAddress});
 	}
 
 	getEdgeServerAddress(serverId: string) {
diff --git a/Engine/Extras/RoboMerge/v3/src/robo/edgebot.ts b/Engine/Extras/RoboMerge/v3/src/robo/edgebot.ts
--- a/Engine/Extras/RoboMerge/v3/src/robo/edgebot.ts
+++ b/Engine/Extras/RoboMerge/v3/src/robo/edgebot.ts
@@ -486,7 +486,7 @@ class EdgeBotImpl extends PerforceStatefulBot {
 			const existingWorkspaceIndex = existingWorkspaceInfos.findIndex((ws) => ws.client == info.targetWorkspaceOverride)
 			if (existingWorkspaceIndex >= 0) {
 				if (existingWorkspaceInfos[existingWorkspaceIndex].IsUnloaded) {
-					await this.p4.reloadWorkspace(info.targetWorkspaceOverride)
+					await this.p4.reloadWorkspace(info.targetWorkspaceOverride, edgeServer.address)
 				}
 				await p4util.cleanWorkspaces(this.p4, [[info.targetWorkspaceOverride, target.branch.rootPath]], edgeServer.address)
 			}
