Remove redundant if

diff --git a/Engine/Extras/RoboMerge/v3/src/robo/nodebot.ts b/Engine/Extras/RoboMerge/v3/src/robo/nodebot.ts
--- a/Engine/Extras/RoboMerge/v3/src/robo/nodebot.ts
+++ b/Engine/Extras/RoboMerge/v3/src/robo/nodebot.ts
@@ -1262,17 +1262,15 @@ export class NodeBot extends PerforceStatefulBot implements NodeBotInterface {
 		this.conflicts.onAlreadyIntegrated(event)
 
 		if (event.change.userRequest && this.slackMessages) {
-			if (this.slackMessages) {
-				const message = `Change ${event.change.cl} was not necessary in ${event.action.branch.name}`
-				let dm: SlackMessage = {
-					text: message,
-					style: SlackMessageStyles.WARNING,
-					channel: '',
-					mrkdwn: true
-				}
-				
-				this.slackMessages.postDM(this.findEmail(event.change.owner!), event.change.cl, this.branch, dm)			
+			const message = `Change ${event.change.cl} was not necessary in ${event.action.branch.name}`
+			let dm: SlackMessage = {
+				text: message,
+				style: SlackMessageStyles.WARNING,
+				channel: '',
+				mrkdwn: true
 			}
+			
+			this.slackMessages.postDM(this.findEmail(event.change.owner!), event.change.cl, this.branch, dm)			
 		}
 	}
 
