Nodes of a graph will now try to process their edges in parallel to maximize concurrent perforce operations.
Restructured the waiting behaviors of nodes. The graph will now try to start processing nodes every checkIntervalSecs (default 30). If the tick took longer than that, it will immediately begin processing again. Each tick will now integrate up to 10 changes per node (down from 50) to allow other edges, pending retries and reconsiders, and other operations to be evaluated more frequently.

diff --git a/Engine/Extras/RoboMerge/v3/src/robo/graphbot.ts b/Engine/Extras/RoboMerge/v3/src/robo/graphbot.ts
--- a/Engine/Extras/RoboMerge/v3/src/robo/graphbot.ts
+++ b/Engine/Extras/RoboMerge/v3/src/robo/graphbot.ts
@@ -167,7 +167,7 @@ export class GraphBot implements GraphInterface, BotEventHandler {
 			this.botlist = [this.autoUpdater, ...this.botlist]
 		}
 
-		this.waitTime = Math.ceil(1000 * this.branchGraph.config.checkIntervalSecs) / this.botlist.length
+		this.waitTime = 1000 * this.branchGraph.config.checkIntervalSecs
 		this.startBotsAsync()
 	}
 
@@ -274,22 +274,22 @@ export class GraphBot implements GraphInterface, BotEventHandler {
 		while (true) {
 			const activity = new Map<string, TickJournal>()
 
-			for (const bot of this.botlist) {
+			const startTime = Date.now()
+
+			let tickBot = async function(bot: Bot, crashRequested: string|null): Promise<[Bot,boolean|Error]> {		
 				bot.isActive = true
 				let ticked = false
 				try {
 					// crashMe API support - simulate a bot crashing and stopping the GraphBot instance
-					if (this.crashRequested) {
-						const errMsg = this.crashRequested
-						this.crashRequested = null
-						throw new Error(errMsg)
+					if (crashRequested) {
+						throw new Error(crashRequested)
 					}
 
 					ticked = await bot.tick()
 				}
 				catch (err) {
-					this.handleNodebotError(bot, err)
-					return
+					bot.isActive = false
+					return [bot,err]
 				}
 				bot.isActive = false
 
@@ -298,19 +298,47 @@ export class GraphBot implements GraphInterface, BotEventHandler {
 					if (bot.tickJournal) {
 						const nodeBot = bot as NodeBot
 						bot.tickJournal.monitored = nodeBot.branch.isMonitored
-						activity.set(nodeBot.branch.upperName, bot.tickJournal)
 					}
 				}
+				return [bot,ticked]
+			}
 
-				if (this._shutdownCb) {
-					this._shutdownCb()
-					this._runningBots = false
-					delete this.eventTriggers
-					this._shutdownCb = null
-					return
+			// The autoupdater needs to fully run before we parallelize the remaining bots
+			if (this.autoUpdater) {
+				const autoUpdateResult = await tickBot(this.autoUpdater, null)
+				var tickResults = [autoUpdateResult, 
+							       ...await Promise.all(this.botlist.slice(1).map(async (bot, index) => tickBot(bot, index == 0 ? this.crashRequested : null)))]
+			}
+			else {
+				tickResults = await Promise.all(this.botlist.map(async (bot, index) => tickBot(bot, index == 0 ? this.crashRequested : null)))
+			}
+
+			this.crashRequested = null
+			
+			let botCrashed = false
+			for (const [bot, result] of tickResults)
+			{
+				if (typeof result === 'boolean') {
+					if (result && bot.tickJournal) {
+						const nodeBot = bot as NodeBot
+						activity.set(nodeBot.branch.upperName, bot.tickJournal)
+					}
+				}
+				else {
+					botCrashed = true
+					this.handleNodebotError(bot, result)
 				}
+			}
+			if (botCrashed) {
+				return
+			}
 
-				await new Promise(done => setTimeout(done, this.waitTime!))
+			if (this._shutdownCb) {
+				this._shutdownCb()
+				this._runningBots = false
+				delete this.eventTriggers
+				this._shutdownCb = null
+				return
 			}
 
 			const errPair = await this.handleRequestedIntegrationsForAllNodes()
@@ -323,6 +351,12 @@ export class GraphBot implements GraphInterface, BotEventHandler {
 			roboAnalytics!.reportActivity(this.branchGraph.botname, activity)
 			roboAnalytics!.reportMemoryUsage('main', process.memoryUsage().heapUsed)
 
+			const duration = Date.now() - startTime;
+			if (duration < this.waitTime!)
+			{
+				await new Promise(done => setTimeout(done, this.waitTime!-duration))
+			}
+
 			// reset tick journals to start counting all events, some of which may happen outside of the bot's tick
 			for (const bot of this.botlist) {
 				if (bot.tickJournal) {
diff --git a/Engine/Extras/RoboMerge/v3/src/robo/nodebot.ts b/Engine/Extras/RoboMerge/v3/src/robo/nodebot.ts
--- a/Engine/Extras/RoboMerge/v3/src/robo/nodebot.ts
+++ b/Engine/Extras/RoboMerge/v3/src/robo/nodebot.ts
@@ -35,7 +35,7 @@ const NAG_EMAIL_MIN_TIME_MINUTES = 60
 const NAG_EMAIL_MIN_TIME_DESCRIPTION = 'an hour'
 const SYNTAX_ERROR_PAUSE_TIMEOUT_SECONDS = 10 * 60
 
-const MAX_CHANGES_TO_PROCESS_BEFORE_YIELDING = 50 // when catching up, we seem to get through 10 changes a minute
+const MAX_CHANGES_TO_PROCESS_BEFORE_YIELDING = 10 // when catching up, we seem to get through 10 changes a minute
 
 const ALLOWED_STOMPABLE_NONBINARY = [
 	/\.collection$/
@@ -249,9 +249,7 @@ export class NodeBot extends PerforceStatefulBot implements NodeBotInterface {
 		if (this.previewMode) {
 			return false
 		}
-		for (const edgeBot of this.edges.values()) {
-			await edgeBot.tick()
-		}
+		await Promise.all(Array.from(this.edges.values(), async edgeBot => edgeBot.tick()))
 
 		// Pre-tick
 
@@ -287,7 +285,7 @@ export class NodeBot extends PerforceStatefulBot implements NodeBotInterface {
 			// reset integration timers - easier to do it here than on every integration code path
 			edgeBot.resetIntegrationTimestamp()
 		}
-	// End Pre-tick
+		// End Pre-tick
 
 		// see if our flow is paused
 		if (this.isManuallyPaused) {
