If p4 change -i times out, retry.

diff --git a/Engine/Extras/RoboMerge/v3/src/common/perforce.ts b/Engine/Extras/RoboMerge/v3/src/common/perforce.ts
--- a/Engine/Extras/RoboMerge/v3/src/common/perforce.ts
+++ b/Engine/Extras/RoboMerge/v3/src/common/perforce.ts
@@ -804,15 +804,32 @@ export class PerforceContext {
 
 		// run the P4 change command
 		this.logger.info("Executing: 'p4 change -i' to create a new CL");
-		const output = await this._execP4(workspace, ['change', '-i'], { stdin: form, quiet: true, edgeServerAddress });
-		// parse the CL out of output
-		const match = output.match(/Change (\d+) created./);
-		if (!match) {
-			throw new Error('Unable to parse new_cl output:\n' + output);
-		}
+		while (true) {
+			try {
+				const output = await this._execP4(workspace, ['change', '-i'], { stdin: form, quiet: true, edgeServerAddress });
+				// parse the CL out of output
+				const match = output.match(/Change (\d+) created./);
+				if (!match) {
+					throw new Error('Unable to parse new_cl output:\n' + output);
+				}
+				// return the changelist number
+				return parseInt(match![1]);
+			}
+			catch (reason) {
+				if (!isExecP4Error(reason)) {
+					throw reason
+				}
 
-		// return the changelist number
-		return parseInt(match[1]);
+				let [err, output] = reason
+				// If perforce timed out try again
+				if (output.includes("Operation took too long")) {
+					this.logger.info("p4 change -i timed out. Retrying.")
+					continue
+				}
+
+				throw err;			
+			}
+		}
 	}
 
 	// integrate a CL from source to destination, resolve, and place the results in a new CL
