Yield changelist processing based on the wait time not number of integrations

diff --git a/Engine/Extras/RoboMerge/v3/src/robo/nodebot.ts b/Engine/Extras/RoboMerge/v3/src/robo/nodebot.ts
--- a/Engine/Extras/RoboMerge/v3/src/robo/nodebot.ts
+++ b/Engine/Extras/RoboMerge/v3/src/robo/nodebot.ts
@@ -35,8 +35,6 @@ const NAG_EMAIL_MIN_TIME_MINUTES = 60
 const NAG_EMAIL_MIN_TIME_DESCRIPTION = 'an hour'
 const SYNTAX_ERROR_PAUSE_TIMEOUT_SECONDS = 10 * 60
 
-const MAX_CHANGES_TO_PROCESS_BEFORE_YIELDING = 10 // when catching up, we seem to get through 10 changes a minute
-
 const ALLOWED_STOMPABLE_NONBINARY = [
 	/\.collection$/
 ]
@@ -346,7 +344,7 @@ export class NodeBot extends PerforceStatefulBot implements NodeBotInterface {
 			this.ticksSinceLastNewP4Commit = 0
 
 			this.headCL = changes[0].change
-			await this._processListOfChanges(availableEdges, changes, MAX_CHANGES_TO_PROCESS_BEFORE_YIELDING)
+			await this._processListOfChanges(availableEdges, changes)
 		}
 		return true
 	}
@@ -1348,7 +1346,7 @@ export class NodeBot extends PerforceStatefulBot implements NodeBotInterface {
 		}
 	}
 
-	private async _processListOfChanges(availableEdges: Map<string, EdgeBot>, allChanges: Change[], maxChangesToProcess: number) {
+	private async _processListOfChanges(availableEdges: Map<string, EdgeBot>, allChanges: Change[]) {
 		// list of changes except reconsiders
 		const changes: Change[] = []
 
@@ -1374,7 +1372,7 @@ export class NodeBot extends PerforceStatefulBot implements NodeBotInterface {
 		// make sure the list is sorted in ascending order
 		changes.sort((a, b) => a.change - b.change)
 
-		const integrationsPerEdge = new Map<string, number>()
+		const startTime = Date.now()
 
 		for (let changeIndex = 0; changeIndex < changes.length; ++changeIndex) {
 			const change = changes[changeIndex]
@@ -1421,19 +1419,10 @@ export class NodeBot extends PerforceStatefulBot implements NodeBotInterface {
 				edgeBot.onNodeProcessedChange(changes, changeIndex, changeResult)
 			}
 
-			// yield if necessary - might be better to make this time-based
-			if (maxChangesToProcess > 0) {
-
-				for (const [targetName, result] of changeResult.edges) {
-					if (result.result === 'ok') {
-						const numIntegrations = (integrationsPerEdge.get(targetName) || 0) + 1
-						if (numIntegrations >= maxChangesToProcess) {
-							this.nodeBotLogger.info(`${targetName} yielding after ${maxChangesToProcess} revisions`)
-							availableEdges.delete(targetName)
-						}
-						integrationsPerEdge.set(targetName, numIntegrations)
-					}
-				}
+			const duration = startTime - Date.now()
+			if ((this.branchGraph.config.checkIntervalSecs * 1000) < duration) {
+				this.nodeBotLogger.info(`${this.fullName} yielding after ${duration}`)
+				return
 			}
 
 			// If we've been paused (or blocked?) while the previous change was being processed, stop now
