Fix yielding

diff --git a/Engine/Extras/RoboMerge/v3/src/robo/nodebot.ts b/Engine/Extras/RoboMerge/v3/src/robo/nodebot.ts
--- a/Engine/Extras/RoboMerge/v3/src/robo/nodebot.ts
+++ b/Engine/Extras/RoboMerge/v3/src/robo/nodebot.ts
@@ -1419,7 +1419,7 @@ export class NodeBot extends PerforceStatefulBot implements NodeBotInterface {
 				edgeBot.onNodeProcessedChange(changes, changeIndex, changeResult)
 			}
 
-			const duration = startTime - Date.now()
+			const duration = Date.now() - startTime
 			if ((this.branchGraph.config.checkIntervalSecs * 1000) < duration) {
 				this.nodeBotLogger.info(`${this.fullName} yielding after ${duration}`)
 				return
