Return a 503 error instead of throwing an undefined property error when using track changes while the bot is starting up

diff --git a/Engine/Extras/RoboMerge/v3/src/robo/ipc.ts b/Engine/Extras/RoboMerge/v3/src/robo/ipc.ts
--- a/Engine/Extras/RoboMerge/v3/src/robo/ipc.ts
+++ b/Engine/Extras/RoboMerge/v3/src/robo/ipc.ts
@@ -182,6 +182,10 @@ export class IPC {
 			return {statusCode: 400, message: `Invalid CL parameter: ${cl}`}
 		}
 
+		if (!this.robo.graph) {
+			return {statusCode: 503, message: 'Service is not ready to process request'}
+		}
+
 		let userTags = new Set<string>()
 		for (const tag in tagsObj) {
 			userTags.add(tag)
