Clean up logging on yield times

diff --git a/Engine/Extras/RoboMerge/v3/src/robo/nodebot.ts b/Engine/Extras/RoboMerge/v3/src/robo/nodebot.ts
--- a/Engine/Extras/RoboMerge/v3/src/robo/nodebot.ts
+++ b/Engine/Extras/RoboMerge/v3/src/robo/nodebot.ts
@@ -1419,9 +1419,9 @@ export class NodeBot extends PerforceStatefulBot implements NodeBotInterface {
 				edgeBot.onNodeProcessedChange(changes, changeIndex, changeResult)
 			}
 
-			const duration = Date.now() - startTime
-			if ((this.branchGraph.config.checkIntervalSecs * 1000) < duration) {
-				this.nodeBotLogger.info(`${this.fullName} yielding after ${duration}`)
+			const duration = Math.round((Date.now() - startTime) / 1000);
+			if (this.branchGraph.config.checkIntervalSecs < duration) {
+				this.nodeBotLogger.info(`${this.fullName} yielding after ${duration}s`)
 				return
 			}
 
