Fix crash processing autoupdater reload of branchspec

diff --git a/Engine/Extras/RoboMerge/v3/src/robo/autobranchupdater.ts b/Engine/Extras/RoboMerge/v3/src/robo/autobranchupdater.ts
--- a/Engine/Extras/RoboMerge/v3/src/robo/autobranchupdater.ts
+++ b/Engine/Extras/RoboMerge/v3/src/robo/autobranchupdater.ts
@@ -117,7 +117,7 @@ export class AutoBranchUpdater implements Bot {
 			}
 		}
 		catch (err) {
-			// if we're in devmode we support files open for add or not in perforce so absord the exception
+			// if we're in devmode we support files open for add or not in perforce so absorb the exception
 			if (!AutoBranchUpdater.config.devMode) {
 				this.abuLogger.printException(err, 'Branch specs: error while querying P4 for changes');
 				return false
diff --git a/Engine/Extras/RoboMerge/v3/src/robo/graphbot.ts b/Engine/Extras/RoboMerge/v3/src/robo/graphbot.ts
--- a/Engine/Extras/RoboMerge/v3/src/robo/graphbot.ts
+++ b/Engine/Extras/RoboMerge/v3/src/robo/graphbot.ts
@@ -163,10 +163,6 @@ export class GraphBot implements GraphInterface, BotEventHandler {
 				this.botlist.push(branch.bot)
 		}
 
-		if (this.autoUpdater) {
-			this.botlist = [this.autoUpdater, ...this.botlist]
-		}
-
 		this.waitTime = 1000 * this.branchGraph.config.checkIntervalSecs
 		this.startBotsAsync()
 	}
@@ -264,6 +260,13 @@ export class GraphBot implements GraphInterface, BotEventHandler {
 
 		this._runningBots = true
 
+		if (this.autoUpdater) {
+			if (!this.autoUpdater.isRunning) {
+				this.botLogger.debug(`Starting bot ${this.autoUpdater.fullNameForLogging}`)
+				this.autoUpdater.start()
+			}
+		}
+
 		for (const bot of this.botlist) {
 			if (!bot.isRunning) {
 				this.botLogger.debug(`Starting bot ${bot.fullNameForLogging}`)
@@ -288,7 +291,6 @@ export class GraphBot implements GraphInterface, BotEventHandler {
 					ticked = await bot.tick()
 				}
 				catch (err) {
-					bot.isActive = false
 					return [bot,err]
 				}
 				bot.isActive = false
@@ -305,14 +307,22 @@ export class GraphBot implements GraphInterface, BotEventHandler {
 
 			// The autoupdater needs to fully run before we parallelize the remaining bots
 			if (this.autoUpdater) {
-				const autoUpdateResult = await tickBot(this.autoUpdater, null)
-				var tickResults = [autoUpdateResult, 
-							       ...await Promise.all(this.botlist.slice(1).map(async (bot, index) => tickBot(bot, index == 0 ? this.crashRequested : null)))]
-			}
-			else {
-				tickResults = await Promise.all(this.botlist.map(async (bot, index) => tickBot(bot, index == 0 ? this.crashRequested : null)))
+				const [_, autoUpdateResult] = await tickBot(this.autoUpdater, null)
+				if (typeof autoUpdateResult !== 'boolean') {
+					this.handleNodebotError(this.autoUpdater, autoUpdateResult)
+					return
+				}
+				if (this._shutdownCb) {
+					this._shutdownCb()
+					this._runningBots = false
+					delete this.eventTriggers
+					this._shutdownCb = null
+					return
+				}
 			}
 
+			const tickResults = await Promise.all(this.botlist.map(async (bot, index) => tickBot(bot, index == 0 ? this.crashRequested : null)));
+
 			this.crashRequested = null
 			
 			let botCrashed = false
