Robomerge will now present users with the option to unlock files when blocked by an exclusive checkout
Exclusive lock slack messages will include the full path.
Stomp will noly be shown on the drop down menu when it is a valid option.

diff --git a/Engine/Extras/RoboMerge/v3/functional_tests/framework/src/framework.ts b/Engine/Extras/RoboMerge/v3/functional_tests/framework/src/framework.ts
--- a/Engine/Extras/RoboMerge/v3/functional_tests/framework/src/framework.ts
+++ b/Engine/Extras/RoboMerge/v3/functional_tests/framework/src/framework.ts
@@ -652,6 +652,61 @@ export abstract class FunctionalTest {
 		}
 	}
 
+	async verifyUnlockRequest(source: string, target: string, edgeState: EdgeState) {
+		if (!edgeState.isBlocked()) {
+			throw new Error('edge must be blocked to unlock!')
+		}
+
+		const conflictCl = edgeState.conflict && edgeState.conflict.change
+		const sourceBranchName = this.fullBranchName(source)
+		const targetBranchName = this.fullBranchName(target)
+
+		// verify unlock
+		const endpoint = OPERATION_URL_TEMPLATE
+			.replace('<bot>', this.botName)
+			.replace('<node>', sourceBranchName)
+			.replace('<op>', 'verifyunlock')
+		const url = `${endpoint}?cl=${conflictCl}&target=${targetBranchName}`
+		let verifyResult: any
+
+		try {
+			const post = bent('POST', 'json', 200, 400)
+			verifyResult = await post(url)
+		}
+		catch (err) {
+			this.error(err)
+			throw new Error(`Verifying Unlock with Url "${url}" returned an error.`)
+		}
+
+		return verifyResult
+	}
+
+	async performUnlockRequest(source: string, target: string, edgeState: EdgeState) {
+		const conflictCl = edgeState.conflict && edgeState.conflict.change
+		const sourceBranchName = this.fullBranchName(source)
+		const targetBranchName = this.fullBranchName(target)
+
+		const unlockEndpoint = OPERATION_URL_TEMPLATE
+			.replace('<bot>', this.botName)
+			.replace('<node>', sourceBranchName)
+			.replace('<op>', 'unlockchanges')
+		const url = `${unlockEndpoint}?cl=${conflictCl}&target=${targetBranchName}`
+
+		const post = bent('POST', 200, 400, 500)
+		let response: bent.BentResponse
+		try {
+			response = await post(url) as bent.BentResponse
+		}
+		catch(err) {
+			this.error(err)
+			throw new Error(`Performing Unlock with Url "${url}" returned an error.`)
+		}
+		return {
+			statusCode: response.statusCode,
+			body: response
+		}
+	}
+
 	async reconsider(source: string, cl: number, target?: string, commandOverride?: string) {
 		const sourceBranchName = this.fullBranchName(source)
 
@@ -756,10 +811,6 @@ export abstract class FunctionalTest {
 		if (!verifyResult.validRequest)
 		{
 			this.warn(verifyResult.message)
-			// this.warn("nonBinaryFilesResolved=" + verifyResult.nonBinaryFilesResolved)
-			// this.warn("remainingAllBinary=" + verifyResult.remainingAllBinary)
-			// this.warn("files=" + (Array.isArray(verifyResult.files)
-				// ? verifyResult.files[0].targetFileName : `no files (${verifyResult.files})`))
 			throw new Error('Stomp verify returned unexpected values')
 		}
 
@@ -771,6 +822,49 @@ export abstract class FunctionalTest {
 		return { verify: verifyResult, stomp: stompResult }
 	}
 
+	async verifyAndPerformUnlock(source: string, target: string, additionalSlackChannel?: string) {
+		const edgeState: EdgeState = await this.getEdgeState(source, target)
+
+		const conflictCl: number | undefined = edgeState.conflict && edgeState.conflict.change
+		if (!conflictCl) {
+			throw new Error('no conflict cl in edge state')
+		}
+
+		const checkSlack = async () => {
+			const slackChannel = this.botName.toLowerCase()
+			const channelsStr = slackChannel + (additionalSlackChannel ? ' and ' + additionalSlackChannel : '')
+
+			this.info(`Ensuring Slack message sent to ${channelsStr} for CL#${conflictCl}`)
+
+			const channels = [slackChannel]
+			if (additionalSlackChannel) {
+				channels.push(additionalSlackChannel)
+			}
+
+			await Promise.all(channels.map(channel =>
+				this.ensureConflictMessagePostedToSlack(source, target, channel)))
+		}
+
+		this.info(`Unlock ${source} -> ${target} @ CL#${conflictCl}`)
+		const [verifyResult] = await Promise.all([
+			this.verifyUnlockRequest(source, target, edgeState),
+			checkSlack()
+		])
+
+		if (!verifyResult.validRequest)
+		{
+			this.warn(verifyResult.message)
+			throw new Error('Unlock verify returned unexpected values')
+		}
+
+		// attempt stomp
+		const stompResult = await this.performUnlockRequest(source, target, edgeState)
+		this.verbose('Waiting for RoboMerge to process Unlock')
+		await this.waitForRobomergeIdle()
+
+		return { verify: verifyResult, stomp: stompResult }
+	}
+
 	async createShelf(source: string, target: string) {
 		const edgeState = await this.getEdgeState(source, target)
 		if (!edgeState.isBlocked()) {
diff --git a/Engine/Extras/RoboMerge/v3/functional_tests/framework/src/tests.ts b/Engine/Extras/RoboMerge/v3/functional_tests/framework/src/tests.ts
--- a/Engine/Extras/RoboMerge/v3/functional_tests/framework/src/tests.ts
+++ b/Engine/Extras/RoboMerge/v3/functional_tests/framework/src/tests.ts
@@ -14,6 +14,7 @@ import { CrossDepotStreamIntegration } from './tests/cross-depot-stream-integrat
 import { EdgeIndependence } from './tests/edge-independence'
 import { EdgeInitialCl } from './tests/edge-initial-cl'
 import { ExclusiveCheckout } from './tests/exclusive-checkout'
+import { ExclusiveCheckoutUnlock } from './tests/exclusive-checkout-unlock'
 import { ExcludeAuthors } from './tests/exclude-authors'
 import { ExcludeAuthorsPerEdge } from './tests/exclude-authors-per-edge'
 import { ForwardCommands, ForwardCommands2 } from './tests/forward-commands'
@@ -217,7 +218,8 @@ async function go() {
 		new UnreachableSkip(p4), // 45
 
 		new PostToAdditionalChannel(p4),
-		new BranchspecTest(p4)
+		new BranchspecTest(p4),
+		new ExclusiveCheckoutUnlock(p4)
 	]
 
 	const TARGET_TEST_DEFS: [string[], string | null][] = [
diff --git a/Engine/Extras/RoboMerge/v3/functional_tests/framework/src/tests/exclusive-checkout-unlock.ts b/Engine/Extras/RoboMerge/v3/functional_tests/framework/src/tests/exclusive-checkout-unlock.ts
new file mode 100644
--- /dev/null
+++ b/Engine/Extras/RoboMerge/v3/functional_tests/framework/src/tests/exclusive-checkout-unlock.ts
@@ -0,0 +1,46 @@
+// Copyright Epic Games, Inc. All Rights Reserved.
+import { FunctionalTest, P4Client, P4Util, RobomergeBranchSpec, Stream } from '../framework'
+
+const streams: Stream[] = [
+	{name: 'Main', streamType: 'mainline'},
+	{name: 'Release', streamType: 'release', parent: 'Main'}
+]
+
+export class ExclusiveCheckoutUnlock extends FunctionalTest {
+	mainSpec: RobomergeBranchSpec
+	releaseUser2Client: P4Client
+
+	async setup() {
+		await this.p4.depot('stream', this.depotSpec())
+		await this.createStreamsAndWorkspaces(streams, undefined, ['testuser1', 'testuser2'])
+
+		const mainClient = this.getClient('Main', 'testuser1')
+		await P4Util.addFileAndSubmit(mainClient, 'test.uasset', 'Dummy content', true)
+		
+		this.releaseUser2Client = this.getClient('Release', 'testuser2')
+		await Promise.all([
+			mainClient.edit('test.uasset'),
+			this.p4.populate(this.getStreamPath('Release'), 'Initial branch of files from Main')
+				.then(() => this.releaseUser2Client.sync())
+		])
+	}
+
+	run() {
+		// second user commits change in Release well first user has file in main checked out
+		return P4Util.editFileAndSubmit(this.releaseUser2Client, 'test.uasset', 'New content')
+	}
+
+	async verify() {
+		await this.ensureBlocked('Release', 'Main')
+        await this.verifyAndPerformUnlock('Release', 'Main')
+        await this.ensureNotBlocked('Release', 'Main')
+
+	}
+
+	getBranches() {
+		return [
+			this.makeForceAllBranchDef('Main', []),
+			this.makeForceAllBranchDef('Release', ['Main'])
+		]
+	}
+}
diff --git a/Engine/Extras/RoboMerge/v3/public/js/boilerplate.js b/Engine/Extras/RoboMerge/v3/public/js/boilerplate.js
--- a/Engine/Extras/RoboMerge/v3/public/js/boilerplate.js
+++ b/Engine/Extras/RoboMerge/v3/public/js/boilerplate.js
@@ -1401,11 +1401,17 @@ function renderActionsCell_Common(actionCell, data, operationFunction, operation
 			}, `Retry merge of ${conflict.cl} and create a shelf in a specified P4 workspace`)
 			
 			// Stomp
-			const stompRequest = `/op/stomp?` + queryParams + location.hash
+			const stompRequest = '/op/stomp?' + queryParams + location.hash
 			const stompOption = createActionOption('Stomp Changes using ' + conflict.cl + toTargetText, function() {
 				window.location.href = stompRequest;
 			}, `Use ${conflict.cl} to stomp binary changes in ${conflict.target}`);
 
+			// Unlock files
+			const unlockFilesRequest = '/op/unlock?' + queryParams
+			const unlockOption = createActionOption('Unlock Files blocking ' + conflict.cl + toTargetText, function() {
+				window.location.href = unlockFilesRequest;
+			}, `Unlock files blocking ${conflict.cl} in ${conflict.target}`);
+			
 			// Can only create shelves for merge conflicts and commit failures
 			if (conflict.kind !== 'Merge conflict' &&
 				conflict.kind !== 'Commit failure') {
@@ -1415,16 +1421,27 @@ function renderActionsCell_Common(actionCell, data, operationFunction, operation
 				shelfOption.attr('data-original-title', `Shelving not available for ${conflict.kind.toLowerCase()}.`)
 			}
 
+			dropdownEntries.push(shelfOption)
+
 			// Can only perform stomps for merge conflicts
-			if (conflict.kind !== 'Merge conflict')
-			{
+			if (conflict.kind === 'Merge conflict') {
+				dropdownEntries.push(stompOption)
+			}
+			else {
 				stompOption.addClass("disabled")
 				stompOption.off('click')
 				stompOption.attr('data-original-title', `Stomp not available for ${conflict.kind.toLowerCase()}.`)
 			}
 
-			dropdownEntries.push(shelfOption)
-			dropdownEntries.push(stompOption)
+			if (conflict.kind === 'Exclusive check-out') {
+				dropdownEntries.push(unlockOption)
+			}
+			else {
+				unlockOption.addClass("disabled")
+				unlockOption.off('click')
+				unlockOption.attr('data-original-title', `Unlock not required for ${conflict.kind.toLowerCase()}.`)
+			}
+
 		}
 
 		// Exposing this section in the case of a blockage but no conflict
diff --git a/Engine/Extras/RoboMerge/v3/public/js/stomp.js b/Engine/Extras/RoboMerge/v3/public/js/stomp.js
--- a/Engine/Extras/RoboMerge/v3/public/js/stomp.js
+++ b/Engine/Extras/RoboMerge/v3/public/js/stomp.js
@@ -54,7 +54,7 @@ function stompVerify() {
             
             // Verify the request branch is paused
             if (!requestedEdge.blockage) {
-                stompFailure(`${requestedEdge.display_name} not currently blocked, no need to skip.`)
+                stompFailure(`${requestedEdge.display_name} not currently blocked, no need to stomp.`)
                 $('#result').append(renderSingleBranchTable(requestedNode))
                 return
             }
diff --git a/Engine/Extras/RoboMerge/v3/public/js/unlock.js b/Engine/Extras/RoboMerge/v3/public/js/unlock.js
new file mode 100644
--- /dev/null
+++ b/Engine/Extras/RoboMerge/v3/public/js/unlock.js
@@ -0,0 +1,253 @@
+// Copyright Epic Games, Inc. All Rights Reserved.
+
+function goHome() {
+    location.href = '/' + location.hash
+}
+
+function unlockVerify() {
+    let queryParams = processQueryParameters(['bot', 'branch', 'cl', 'target'])
+    if (!queryParams) return;
+
+    let requestedBotName = queryParams["bot"]
+    let requestedBranchName = queryParams["branch"]
+    let targetBranchName = queryParams["target"]
+    let requestedBranchCl = parseInt(queryParams["cl"], 10)
+    
+    if (isNaN(requestedBranchCl)) {
+        unlockFailure(`CL ${queryParams["cl"]} not a number.`)
+        return
+    }
+
+    // Filled in if/when user clicks the verify button
+    let unlockOperation = $.Deferred()
+    
+    // Get requested branch infomation and setup the page data
+    let visualizationOperation = getBranch(requestedBotName, requestedBranchName, function(data) {
+        try {
+            // Ensure we have data
+            if (!data) {
+                unlockFailure(`Couldn't retrieve data for ${requestedBotName}:${requestedBranchName}`)
+                return
+            }
+            
+            let requestedNode = data.branch
+
+            // Verify we got a node
+            if (!requestedNode) {
+                let errText = `Could not find matching branch for ${requestedBotName}:${requestedBranchName}.`
+                if (data.message) {
+                    errText += `\n${data.message}`
+                }
+                unlockFailure(errText)
+                return
+            }
+
+            let requestedEdge = requestedNode.edges[targetBranchName.toUpperCase()]
+
+            // Ensure this node actually has this requested edge
+            if (!requestedEdge) {
+                unlockFailure(`${requestedBotName}:${requestedNodeName} has no edge for "${targetBranchName}".`)
+                $('#result').append(renderSingleBranchTable(requestedNode))
+                return
+            }
+        
+            
+            // Verify the request branch is paused
+            if (!requestedEdge.blockage) {
+                unlockFailure(`${requestedEdge.display_name} not currently blocked, no need to unlock.`)
+                $('#result').append(renderSingleBranchTable(requestedNode))
+                return
+            }
+
+
+            // Ensure the current pause is applicable to the CL
+            if (requestedEdge.blockage.change !== requestedBranchCl) {
+                displayWarningMessage(`${requestedEdge.display_name} currently blocked, but not at requested CL ${requestedBranchCl}. Performing no action.`, false)
+                $('#singleBranchDisplay').append(renderSingleBranchTable(requestedNode)).fadeIn("fast", "swing")
+                return
+            }
+
+            // Passed all safety checks. Display branch info and unlock warning, and begin verifying the request
+            $('#branchPreviewBeforeVerify').append(renderSingleBranchTable(requestedNode, requestedEdge))
+
+            $('.sourceName').each(function() {
+                $(this).text(requestedBranchName)
+            })
+            $('.targetName').each(function() {
+                $(this).text(targetBranchName)
+            })
+            $('.changelist').each(function() {
+                $(this).text(requestedBranchCl)
+            })
+
+            // Data all verified, time to perform the node operation.
+            let unlockQueryData = {
+                cl: requestedEdge.blockage.change,
+                target: targetBranchName
+            }
+
+            let verifyOperation = nodeAPIOp(requestedNode.bot, requestedNode.def.name, "/verifyunlock?" + toQuery(unlockQueryData))
+
+            // If verification succeeded:
+            verifyOperation.done(function(success) {
+                // Unlock returns a JSON payload on success
+                const unlockJson = JSON.parse(success)
+                console.log(`Unlock Verification message: ${unlockJson.message}`)
+
+                // Visualize the json
+                visualizeUnlockVerification(requestedBranchCl, unlockJson)
+
+                if (unlockJson.validRequest) {
+                    $('#afterVerificationResultText').html('<span><i class="fas fa-check-circle"></i></span> Unlock Verification Success!')
+                    $('#afterVerificationResultText').css('color', 'green');
+                } else {
+                    $('#afterVerificationResultText').html('<span><i class="fas fa-exclamation-triangle"></i></span> Unlock Verification Returned Issues.')
+                    $('#afterVerificationResultText').css('color', 'red');
+                }
+
+                transitionPostVerification()
+
+                // If we're a valid request, add form buttons to formally request an unlock
+                if (unlockJson.validRequest) {
+                    const formButtonDiv = $('#formButtons')
+
+                    formButtonDiv.append(
+                        $('<h3 style="text-align: center">').text("Use the button below to proceed with the unlock operation:")
+                    )
+                    // Return to Robomerge homepage
+                    let cancelButton = $('<button type="button" class="btn btn-lg btn-info" style="margin:.3em;">').text(`Cancel`).appendTo(formButtonDiv)
+                    cancelButton.click(function() {
+                        goHome()
+                    })
+
+                    // Perform unlock
+                    let unlockButton = $('<button type="button" class="btn btn-lg btn-primary" style="margin:.3em;">').text(`Perform Unlock`).appendTo(formButtonDiv)
+                    unlockButton.click(function() {
+                        unlockButton.attr('disabled', true)
+                        unlockButton.text(`Requesting unlock with CL ${requestedBranchCl}...`)
+                        transitionUnlockChanges(requestedBranchCl)
+
+                        nodeAPIOp(requestedNode.bot, requestedNode.def.name, "/unlockchanges?" + toQuery(unlockQueryData))
+                            .done(function (success) {
+                                unlockOperation.resolve(success)
+                            })
+                            .fail(function(jqXHR, textStatus, errMsg) {
+                                unlockOperation.reject(jqXHR, textStatus, errMsg)
+                            })
+                    })
+                } else {
+                    $('#returnbutton').removeClass("initiallyHidden")
+                }
+            })
+
+            // If verification failed:
+            verifyOperation.fail(function(jqXHR, _textStatus, errMsg) {
+                transitionPostVerification()
+
+                $('#returnbutton').removeClass("initiallyHidden")
+                const errText = `${jqXHR && jqXHR.status ? jqXHR.status + ": " : ""}${jqXHR && jqXHR.responseText ? jqXHR.responseText : errMsg}`
+                unlockFailure(`Encountered issues verifying files from CL ${requestedEdge.blockage.change} to unlock changes in ${targetBranchName}: "${errText}" If you think this is an error, please contact Robomerge support.`)
+                
+                // Call getBranchList a second time to display most up-to-date data
+                getBranch(requestedNode.bot, requestedNode.def.name, function(data) {
+                    const upToDateEdgeData = data.edges && data.edges[targetBranchName.toUpperCase()] ? data.edges[targetBranchName.toUpperCase()] : null
+                    $('#afterVerification').append(renderSingleBranchTable(data.branch, upToDateEdgeData))
+                })
+            })
+        }
+        catch (err) {
+            unlockFailure(`Error encountered displaying unlock verification results: ${err}`)
+        }
+    })
+
+    // Logic controlling unlock operation
+    $.when(unlockOperation).done(function(success) {
+        try {
+            transitionDisplayUnlockResults()
+            unlockSuccess(success)
+        }
+        catch (err) {
+            unlockFailure(`Error encountered displaying unlock files results: ${err}`)
+        }
+        
+    }).fail(function(jqXHR, textStatus, errMsg) {
+        try {
+            transitionDisplayUnlockResults()
+            unlockFailure(`${textStatus}: ${errMsg}`)
+        }
+        catch (err) {
+            unlockFailure(`Error encountered displaying unlock files failure: ${err}`)
+        }
+    })
+}
+
+function transitionPostVerification() {
+    // Remove before verify animation, fade in new visualization
+    $('#beforeVerification').fadeOut("fast", "swing", function() { $('#afterVerification').fadeIn("fast") })
+    $('#loadingDiv').fadeOut("fast", "swing")
+}
+
+function transitionUnlockChanges(changelist) {
+    $('#loadingText').text(`Performing unlock with CL ${changelist}...`)
+    $('#afterVerification').fadeOut("fast", "swing", function() { if ($('#returnbutton').hasClass("initiallyHidden")) { $('#loadingDiv').fadeIn("fast") }})
+}
+
+function transitionDisplayUnlockResults() {
+    $('#loadingDiv').fadeOut("fast", "swing", function() { $('#afterUnlock').fadeIn("fast") })
+    $('#returnbutton').removeClass("initiallyHidden")
+}
+
+// Create visualization of the file, with links to swarm
+function visualizeAuthor(author, files) {
+    const authorDiv = $(`<div id="${author}" class="unlock-visual">`)
+    if (author.toLowerCase() === robomergeUser.userName.toLowerCase()) {
+        authorDiv.append($('<h3>').html("Locked by <strong>you</strong>"))
+    } else {
+        authorDiv.append($('<h3>').html(`Locked by ${author}`))
+    }
+    
+    // Begin compiling the list of relevant files in this changelist
+    let fileList = $('<ul>')
+    for (const file of files) {
+        fileList.append($('<li>').html(file))
+    }
+    fileList.appendTo(authorDiv)
+
+    return authorDiv
+}
+
+function visualizeUnlockVerification(requestedBranchCl, unlockJson) {
+    // Debug
+    const prettyJson = JSON.stringify(unlockJson, null, 2)
+    console.log(`Unlock Verification Debug JSON:\n${prettyJson}`)
+    let linkToSwarm = makeClLink(requestedBranchCl, `CL ${requestedBranchCl}`)
+
+    let authorDict = {}
+    for (const file of unlockJson.files) {
+
+        if (!authorDict[file.user]) {
+            authorDict[file.user] = {
+                files: [file.depotPath]
+            }
+        }
+        else {
+            authorDict[file.user].files.push(file.depotPath)
+        }
+    }
+
+    Object.keys(authorDict).forEach(function(key) {
+        $('#resultVisualization').append(visualizeAuthor(key, authorDict[key].files))
+    })
+}
+
+function unlockFailure(message) {
+    displayErrorMessage(message, false)
+    $('#loadingDiv').fadeOut("fast", "swing")
+    $('#returnbutton').removeClass("initiallyHidden")
+}
+
+function unlockSuccess(message) {
+    displaySuccessfulMessage(message, false)
+    $('#loadingDiv').fadeOut("fast", "swing")
+    $('#returnbutton').removeClass("initiallyHidden")
+}
\ No newline at end of file
diff --git a/Engine/Extras/RoboMerge/v3/public/unlock.html b/Engine/Extras/RoboMerge/v3/public/unlock.html
new file mode 100644
--- /dev/null
+++ b/Engine/Extras/RoboMerge/v3/public/unlock.html
@@ -0,0 +1,77 @@
+<!DOCTYPE html>
+<html>
+<head>
+<title>ROBOMERGE</title>
+<link rel="icon" href="/img/robo-r.png" type="image/png">
+<link rel="stylesheet" href="/css/fontawesome.5.9.0.all.min.css">
+<link rel="stylesheet" href="/css/bootstrap.min.css">
+<link rel="stylesheet" href="/css/main.css">
+<link rel="stylesheet" href="/css/table.css">
+
+	<script src="/js/jquery-3.4.1.min.js"></script>
+	<script src="/js/bootstrap.bundle.min.js"></script>
+
+<script src="/js/boilerplate.js"></script>
+<script src="/js/unlock.js"></script>
+<script language="javascript">
+
+    $(function() {
+        unlockVerify()
+    });
+
+</script>
+</head>
+<body>
+
+<div style="margin-bottom: 1em;">
+    <h2 class="centered-text" style="text-align:center;  margin-bottom: .5em;"><a title="What's This?" data-toggle="collapse" href="#whatsthis" aria-expanded="false">Unlocking Changes</a></h2>
+    <div id="whatsthis" class="collapse">
+        <div class="card card-single border-info centered-text" style="margin:auto">
+            <div class="card-header">What's Unlock Changes?</div>
+            <div class="card-body centered-text">
+                <p>When a merge is blocked by files that are exclusively locked, Robomerge will offer to 'unlock' those files.</p>
+                <p>Once your unlock request is verified and you elect to unlock, the users who had those files locked will be notified in slack.</p>
+                <p>Please report any issues to <a href="https://epic.slack.com/messages/dt-robomerge-help-ext/"><code>#dt-robomerge-help-ext</code></a> on Slack.</p>
+            </div>
+        </div>
+    </div>
+</div>
+
+<!-- This div will hold everything to display to the user while the lengthy verification process starts. -->
+<div id="beforeVerification" class="text" >
+    <div id="branchPreviewBeforeVerify"></div>
+</div>
+
+<div id="loadingDiv" class="card card-single border-dark loading-box">
+    <div class="text-center card-body">
+        <h1 id="loadingText" class="align-self-center">Verifying Unlock Request...</h1>
+        <img id="loadingimg" class="mx-auto d-block" style="width:128px;" src="/img/robo-r-rotating.gif">
+    </div>
+</div>
+
+<!-- This div will hold the results of the verification process. -->
+<div class="text" id="afterVerification" style="display:none;">
+   <div id="verificationComplete">
+        <!-- This heading gets changed by the javascript. If the user sees this text without the Javascript being changed, it should warn the user. -->
+        <h2 class="centered-text" id="afterVerificationResultText" style="text-align:center; margin-bottom: .5em;">Unlock Verification May Have Encountered Issues.</h2>
+
+        <!-- Dynamically generated div to hold visualization of unlock verification JSON -->
+        <div id="resultVisualization"></div>
+    </div>
+
+    <div id="formButtons" style="text-align:center;"></div>
+
+    <!-- End afterVerification -->
+</div>
+
+<div class="centered-text" id="afterUnlock" style="display:none;">
+    <h2 style="text-align:center; margin-bottom: .5em;">Unlock Request Processed.</h2>
+    <p>Robomerge will unlock the blocking files and retry the integration.</p>
+</div>
+
+<div id="singleBranchDisplay" style="display:none;"></div>
+
+<button id="returnbutton" class="centered-text initiallyHidden btn btn-lg btn-block btn-info" onclick="goHome()" style="margin-top: 2em;">Return to RoboMerge homepage</button>
+
+</body>
+</html>
\ No newline at end of file
diff --git a/Engine/Extras/RoboMerge/v3/src/common/analytics.ts b/Engine/Extras/RoboMerge/v3/src/common/analytics.ts
--- a/Engine/Extras/RoboMerge/v3/src/common/analytics.ts
+++ b/Engine/Extras/RoboMerge/v3/src/common/analytics.ts
@@ -247,7 +247,7 @@ export function runTests(parentLogger: ContextualLogger) {
 
 	// test up
 	unitTestLogger.info('... activity counters ...')
-	test.updateActivityCounters({merges: 50, stompQueued: 11, traces: 13})
+	test.updateActivityCounters({merges: 50, stompQueued: 11, traces: 13, unlockQueued: 2})
 	test.flush()
 
 	test.stop()
diff --git a/Engine/Extras/RoboMerge/v3/src/common/perforce.ts b/Engine/Extras/RoboMerge/v3/src/common/perforce.ts
--- a/Engine/Extras/RoboMerge/v3/src/common/perforce.ts
+++ b/Engine/Extras/RoboMerge/v3/src/common/perforce.ts
@@ -60,6 +60,7 @@ export interface OpenedFileRecord {
 	action: string
 	type: string
 	user: string
+	client: string
 
 	movedFile?: string
 }
@@ -1147,6 +1148,25 @@ export class PerforceContext {
 		return this._execP4Ztag(workspace, args) as Promise<OpenedFileRecord[]>
 	}
 
+	async revertFile(file: string) {
+		let opened = await this.opened(null, file, true)
+		if (opened.length == 0) {
+			opened = await this.opened(null, file)
+			if (opened.length == 0) {
+				this.logger.warn(`${file} does not appear to be open, cannot revert`)
+				return null
+			}
+		}
+		const client = opened[0].client
+		return this.revertFiles([file],client)
+	}
+
+	async revertFiles(files: string[], client: string) {
+		const edgeServer = await this.getWorkspaceEdgeServer(client)
+		const args = ['revert', '-C', client, ...files]
+		return this._execP4Ztag(null, args, {edgeServerAddress: edgeServer?.address})
+	}
+
 	// revert a CL deleting any files marked for add
 	// output format is just error or not
 	async revert(roboWorkspace: RoboWorkspace, changelist: number, opts?: string[], edgeServerAddress?: string) {
diff --git a/Engine/Extras/RoboMerge/v3/src/robo/bot-interfaces.ts b/Engine/Extras/RoboMerge/v3/src/robo/bot-interfaces.ts
--- a/Engine/Extras/RoboMerge/v3/src/robo/bot-interfaces.ts
+++ b/Engine/Extras/RoboMerge/v3/src/robo/bot-interfaces.ts
@@ -1,6 +1,6 @@
 // Copyright Epic Games, Inc. All Rights Reserved.
 
-import { Branch, BranchArg, BranchGraphInterface, OperationResult, StompVerification } from "./branch-interfaces";
+import { Branch, BranchArg, BranchGraphInterface, OperationResult, StompVerification, UnlockVerification } from "./branch-interfaces";
 import { BlockagePauseInfoMinimal, BranchStatus } from "./status-types"
 import { TickJournal } from "./tick-journal";
 
@@ -29,6 +29,9 @@ export interface NodeBotIPC extends BotIPC {
 
 	verifyStomp(changeCl: number, targetBranchName: string): Promise<StompVerification>
 	stompChanges(owner: string, changeCl: number, targetBranchName: string): Promise<OperationResult>
+
+	verifyUnlock(changeCl: number, targetBranchName: string): Promise<UnlockVerification>
+	unlockChanges(owner: string, changeCl: number, targetBranchName: string): Promise<OperationResult>
 }
 
 
diff --git a/Engine/Extras/RoboMerge/v3/src/robo/branch-interfaces.ts b/Engine/Extras/RoboMerge/v3/src/robo/branch-interfaces.ts
--- a/Engine/Extras/RoboMerge/v3/src/robo/branch-interfaces.ts
+++ b/Engine/Extras/RoboMerge/v3/src/robo/branch-interfaces.ts
@@ -52,10 +52,32 @@ export interface StompVerification extends OperationResult {
 	nonBinaryFilesResolved?: boolean // Warning to the user that non-binary files are included in changelist, but passed the first resolve
 	svFiles?: StompVerificationFile[] // Array of remaining files
 
-	// Convinence booleans to alert user to problems in verification result
+	// Convenience booleans to alert user to problems in verification result
 	remainingAllBinary?: boolean // Check to see if unresolved non-binary files are remaining -- we shouldn't stomp those!
 }
 
+export interface ExclusiveFile {
+	depotPath: string
+	name: string
+	user: string
+	client: string
+}
+
+export interface ExclusiveLockUser {
+    user: string
+    userEmail: Promise<string | null>
+}
+
+export interface ExclusiveLockInfo {
+	exclusiveLockUsers: ExclusiveLockUser[]
+	lockedFiles: ExclusiveFile[]
+}
+
+export interface UnlockVerification extends OperationResult {
+	validRequest?: boolean
+	lockedFiles?: ExclusiveFile[] // Array of locked files
+}
+
 export interface BranchGraphInterface {
 	botname: string
 	config: BotConfig
@@ -146,6 +168,7 @@ export interface ChangeInfo extends TargetInfo {
 		address: string
 	}
 	targetWorkspaceOverride?: string
+	analyzeAllIntegrationErrors?: boolean
 	overriddenCommand: string
 	macros: string[]
 }
diff --git a/Engine/Extras/RoboMerge/v3/src/robo/edgebot.ts b/Engine/Extras/RoboMerge/v3/src/robo/edgebot.ts
--- a/Engine/Extras/RoboMerge/v3/src/robo/edgebot.ts
+++ b/Engine/Extras/RoboMerge/v3/src/robo/edgebot.ts
@@ -8,7 +8,7 @@ import { Recipients } from "../common/mailer";
 import { Change, coercePerforceWorkspace, ConflictedResolveNFile, EditChangeOpts, EXCLUSIVE_CHECKOUT_REGEX, getRootDirectoryForBranch, IntegrationSource, IntegrationTarget, isExecP4Error, OpenedFileRecord, PerforceContext } from "../common/perforce";
 import { VersionReader } from "../common/version";
 import { EdgeBotInterface, IPCControls, ReconsiderArgs } from "./bot-interfaces";
-import { AlreadyIntegrated, Branch, ChangeInfo, ConflictingFile, Failure, MergeAction, PendingChange } from "./branch-interfaces";
+import { AlreadyIntegrated, Branch, ChangeInfo, ConflictingFile, ExclusiveFile, Failure, MergeAction, PendingChange } from "./branch-interfaces";
 import { EdgeOptions } from "./branchdefs";
 import { PersistentConflict } from "./conflict-interfaces";
 import { BotEventTriggers } from "./events";
@@ -259,13 +259,13 @@ class EdgeBotImpl extends PerforceStatefulBot {
 		return results
 	}
 
-	private async analyzeIntegrationError(errors: string[]) {
-		if (errors.length > MAX_INTEGRATION_ERRORS_TO_ANALYZE) {
+	private async analyzeIntegrationError(errors: string[], analyzeAllIntegrationErrors?: boolean) {
+		if (!analyzeAllIntegrationErrors && errors.length > MAX_INTEGRATION_ERRORS_TO_ANALYZE) {
 			this.edgeBotLogger.error(`Integration error: ${errors.length} files, checking first ${MAX_INTEGRATION_ERRORS_TO_ANALYZE}`)
 		}
 		
 		const openedRequests: [RegExpMatchArray, Promise<OpenedFileRecord[]>, Promise<OpenedFileRecord[]>][] = []
-		for (const err of errors.slice(0, MAX_INTEGRATION_ERRORS_TO_ANALYZE)) {
+		for (const err of analyzeAllIntegrationErrors ? errors : errors.slice(0, MAX_INTEGRATION_ERRORS_TO_ANALYZE)) {
 			const match = err.match(EXCLUSIVE_CHECKOUT_REGEX)
 			if (match) {
 				openedRequests.push([match, this.p4.opened(null, match[1] + match[2], true), this.p4.opened(null, match[1] + match[2])])
@@ -277,12 +277,12 @@ class EdgeBotImpl extends PerforceStatefulBot {
 			const recs = await exclusiveReq
 			if (recs.length > 0) {
 				// should only be one, since we're looking for exclusive check-out errors
-				results.push({name: match[2], user: recs[0].user})
+				results.push({depotPath: match[1] + match[2], name: match[2], user: recs[0].user, client: recs[0].client})
 			} else {
 				const recs = await addReq
 				if (recs.length > 0) {
 					// should only be one, since we're looking for exclusive check-out errors
-					results.push({name: match[2], user: recs[0].user})
+					results.push({depotPath: match[1] + match[2], name: match[2], user: recs[0].user, client: recs[0].client})
 				}
 			}
 		}
@@ -554,19 +554,19 @@ class EdgeBotImpl extends PerforceStatefulBot {
 		}
 
 		const errors = results as string[]
-		const exclusiveFiles = await this.analyzeIntegrationError(errors)
+		const exclusiveFiles = await this.analyzeIntegrationError(errors, info.analyzeAllIntegrationErrors)
 
 		const description = errors.join('\n')
 		let failure: Failure | null = null
 
 		if (exclusiveFiles.length > 0) {
 			// will need to store the exclusive file if we want to @ people in Slack
-			const exclCheckoutMessages = exclusiveFiles.map(exc => `${exc.name} checked out by ${exc.user}`)
-			if (errors.length > MAX_INTEGRATION_ERRORS_TO_ANALYZE) {
+			const exclCheckoutMessages = exclusiveFiles.map(exc => `${exc.depotPath} checked out by ${exc.user}`)
+			if (!info.analyzeAllIntegrationErrors && errors.length > MAX_INTEGRATION_ERRORS_TO_ANALYZE) {
 				exclCheckoutMessages.push(`... and up to ${errors.length - MAX_INTEGRATION_ERRORS_TO_ANALYZE} more`)
 			}
-			const exclCheckoutUsers = Array.from(new Set(exclusiveFiles.map(exc => `${exc.user.toLowerCase()}`))).map(user => ({user, userEmail: this.p4.getEmail(user)}))
-			failure = { kind: 'Exclusive check-out', description, summary: exclCheckoutMessages.join('\n'), additionalInfo: exclCheckoutUsers }
+			const exclusiveLockUsers = Array.from(new Set(exclusiveFiles.map(exc => `${exc.user.toLowerCase()}`))).map(user => ({user, userEmail: this.p4.getEmail(user)}))
+			failure = { kind: 'Exclusive check-out', description, summary: exclCheckoutMessages.join('\n'), additionalInfo: {exclusiveLockUsers,exclusiveFiles} }
 		}
 		else {
 			failure  = { kind: 'Integration error', description }
@@ -1138,8 +1138,3 @@ export class EdgeBot
 		}
 	}
 }
-
-interface ExclusiveFile {
-	name: string
-	user: string
-}
\ No newline at end of file
diff --git a/Engine/Extras/RoboMerge/v3/src/robo/graphbot.ts b/Engine/Extras/RoboMerge/v3/src/robo/graphbot.ts
--- a/Engine/Extras/RoboMerge/v3/src/robo/graphbot.ts
+++ b/Engine/Extras/RoboMerge/v3/src/robo/graphbot.ts
@@ -434,7 +434,8 @@ export class GraphBot implements GraphInterface, BotEventHandler {
 				acknowledgeUrl: OperationUrlHelper.createAcknowledgeUrl(this.externalUrl, 'botname', 'sourcebranch', '0'),
 				createShelfUrl: OperationUrlHelper.createCreateShelfUrl(this.externalUrl, 'botname', 'sourcebranch', '0', 'targetbranch', 'targetstream'),
 				skipUrl: OperationUrlHelper.createSkipUrl(this.externalUrl, 'botname', 'sourcebranch', '0', 'targetbranch'),
-				stompUrl: OperationUrlHelper.createStompUrl(this.externalUrl, 'botname', 'sourcebranch', '0', 'targetbranch')
+				stompUrl: OperationUrlHelper.createStompUrl(this.externalUrl, 'botname', 'sourcebranch', '0', 'targetbranch'),
+				unlockUrl: OperationUrlHelper.createUnlockUrl(this.externalUrl, 'botname', 'sourcebranch', '0', 'targetbranch')
 			}
 			return urls
 		}
diff --git a/Engine/Extras/RoboMerge/v3/src/robo/ipc.ts b/Engine/Extras/RoboMerge/v3/src/robo/ipc.ts
--- a/Engine/Extras/RoboMerge/v3/src/robo/ipc.ts
+++ b/Engine/Extras/RoboMerge/v3/src/robo/ipc.ts
@@ -713,6 +713,53 @@ export class IPC {
 			this.ipcLogger.error('Error processing stomp: ' + operationResult.message)
 			return { statusCode: 500, message: operationResult.message }
 
+		// Requires: cl, target
+		case 'verifyunlock':
+			cl = parseInt(query.cl)
+			if (isNaN(cl)) {
+				return { statusCode: 400, message: 'Invalid CL parameter: ' + cl }
+			}
+
+			target = query.target
+			if (!target) {
+				return { statusCode: 400, message: 'Target parameter is required' }
+			}
+
+			// Attempt to stomp changes
+			let unlockVerification = await bot.verifyUnlock(cl, target)
+			if (unlockVerification.success) {
+				return { 
+					statusCode: 200,
+					message: JSON.stringify({
+						message: unlockVerification.message,
+						files: unlockVerification.lockedFiles,
+						validRequest: unlockVerification.validRequest
+					} )
+				}
+			}
+
+			this.ipcLogger.error('Error verifying unlock: ' + unlockVerification.message)
+			return { statusCode: 500, message: unlockVerification.message }
+
+		case 'unlockchanges':
+			cl = parseInt(query.cl)
+			if (isNaN(cl)) {
+				return { statusCode: 400, message: 'Invalid CL parameter: ' + cl }
+			}
+
+			target = query.target
+			if (!target) {
+				return { statusCode: 400, message: 'Target parameter is required' }
+			}
+
+			// Attempt to unlock changes
+			operationResult = await bot.unlockChanges(query.who, cl, target)
+			if (operationResult.success) {
+				return { statusCode: 200, message: operationResult.message }
+			}
+			this.ipcLogger.error('Error processing unlock: ' + operationResult.message)
+			return { statusCode: 500, message: operationResult.message }
+
 		case 'bypassgatewindow':
 			const sense = query.sense.toLowerCase().startsWith('t');
 			const prefix = sense ? 'en' : 'dis'
diff --git a/Engine/Extras/RoboMerge/v3/src/robo/nodebot.ts b/Engine/Extras/RoboMerge/v3/src/robo/nodebot.ts
--- a/Engine/Extras/RoboMerge/v3/src/robo/nodebot.ts
+++ b/Engine/Extras/RoboMerge/v3/src/robo/nodebot.ts
@@ -10,7 +10,7 @@ import { IPCControls, NodeBotInterface, QueuedChange, ReconsiderArgs } from './b
 import { ApprovalOptions, BotConfig, EdgeOptions, NodeOptions } from './branchdefs'
 import { BlockagePauseInfo, BranchStatus } from './status-types';
 import { AlreadyIntegrated, Blockage, Branch, BranchArg, BranchGraphInterface, ChangeInfo, EndIntegratingToGateEvent, Failure } from './branch-interfaces';
-import { MergeAction, OperationResult, PendingChange, resolveBranchArg, StompedRevision, StompVerification, StompVerificationFile }  from './branch-interfaces';
+import { MergeAction, OperationResult, PendingChange, resolveBranchArg, StompedRevision, StompVerification, StompVerificationFile, UnlockVerification }  from './branch-interfaces';
 import { Conflicts } from './conflicts';
 import { EdgeBot, EdgeMergeResults } from './edgebot';
 import { BotEventTriggers } from './events';
@@ -958,7 +958,7 @@ export class NodeBot extends PerforceStatefulBot implements NodeBotInterface {
 			return { success: false, message: "Missing parameters for stompChanges node operation" }
 		}
 
-		// Reverify, but don't revert on a successfull verification!
+		// Reverify, but don't revert on a successful verification!
 		const verifyResult = await this.verifyStomp(changeCl, targetBranchName)
 
 		if (!verifyResult.success) {
@@ -999,6 +999,126 @@ export class NodeBot extends PerforceStatefulBot implements NodeBotInterface {
 
 	}
 
+	async verifyUnlock(changeCl: number, targetBranchName: string) : Promise<UnlockVerification> {
+		this._log_action(`Verifying unlock request of CL ${changeCl} to ${targetBranchName}`, "verbose")
+		if (!changeCl || !targetBranchName) {
+			return { success: false, message: "Missing parameters for unlock operation" }
+		}
+
+		const targetBranch = this._getBranch(targetBranchName)
+
+		if (!targetBranch) {
+			return { success: false, message: `Unable to retrieve branch infomation for "${targetBranchName}". Unable to unlock, please contact Robomerge support."` }
+		}
+
+		const edge = this.getImmediateEdge(targetBranch)
+
+		if (!edge) {
+			return { success: false, message: `This bot has no edge for specified branch "${targetBranch.name}" on ${this.fullName}`}
+		}
+
+		// Ensure we're paused on a non-manual pause blockage
+		if (!edge.isBlocked) {
+			return { success: false, message: "Branch needs to have conflict to unlock files." }
+		}
+
+		// Get the pause info of the edge
+		const blockageInfo = edge.pauseState.blockagePauseInfo!
+
+		// We've ensured we're on a blockage now. Make sure the CL is still valid.
+		if (blockageInfo.change !== changeCl) {
+			return { success: false, message: `Requested change CL "${changeCl}" does not match blockage change CL "${blockageInfo.change}"` }
+		}
+
+		let conflict = this.conflicts.getConflictByChangeCl(changeCl, targetBranch)
+
+		if (!conflict) {
+			return { success: false, message: `Unable to retrieve conflict for CL ${changeCl} merging to ${targetBranch.name}. Unable to unlock, please contact Robomerge support.` }
+		}
+		if (conflict.kind !== 'Exclusive check-out') {
+			return { success: false, message: `Unlock only usable when edge is blocked by exclusive check-out ${targetBranch.name}. ` + 
+				`If you thinks this is an error, contact Robomerge support.` }
+		}
+
+		return { success: true, validRequest: true, message: "Unlock verified", lockedFiles: blockageInfo.additionalInfo.exclusiveFiles}
+	}
+
+	// First, re-verify the unlock request
+	// If verified, we should have a list of files to unlock
+	async unlockChanges(owner: string, changeCl: number, targetBranchName: string): Promise<OperationResult> {
+		// Ensure this is a valid request
+		// Parameter validation (should not be blank)
+		if (!owner || !changeCl || !targetBranchName) {
+			return { success: false, message: "Missing parameters for unlockChanges node operation" }
+		}
+
+		const verifyResult = await this.verifyUnlock(changeCl, targetBranchName)
+
+		if (!verifyResult.success) {
+			return { success: false, message: `unlockChanges Request Failed Due to Bad Verification: ${verifyResult.message}` }
+		}
+
+		if (!verifyResult.validRequest) {
+			return { success: false, message: `unlockChanges Request Cannot Be Completed: ${verifyResult.message}` }
+		}
+
+		this.nodeBotLogger.info(`Unlock Request for cl ${changeCl} by ${owner} validated, unlocking files`)
+
+		let filesByClient = new Map<string,string[]>()
+		let filesByUser = new Map<string,string[]>()
+		for (const lockedFile of verifyResult.lockedFiles!) {
+			filesByClient.set(lockedFile.client, [...(filesByClient.get(lockedFile.client) || []), lockedFile.depotPath])
+			filesByUser.set(lockedFile.user, [...(filesByUser.get(lockedFile.user) || []), lockedFile.depotPath])
+		}
+
+		filesByClient.forEach((files,client) => this.p4.revertFiles(files,client))
+
+		if (this.slackMessages) {
+			
+			filesByUser.forEach(async (files, author) => {
+
+				let userToNotify = `@${author}`
+				const emailAddress = await this.findEmail(author)
+				if (emailAddress)
+				{
+					const user = await this.slackMessages!.getSlackUser(emailAddress)
+					if (user) {
+						userToNotify = `<@${user}>`
+					}
+				}
+				const unlockMessage: SlackMessage = {
+					text: `${userToNotify} - ${owner} has unlocked the following files:\n${files.join('\n')}`,
+					style: SlackMessageStyles.DANGER,
+					channel: this.branchGraph.config.slackChannel,
+					mrkdwn: true
+				}
+
+				this.slackMessages!.postReply(changeCl, targetBranchName, unlockMessage)
+			})
+		}
+
+		if (this.tickJournal) {
+			// Add number of successful unlocks
+			++this.tickJournal.unlockQueued
+		}
+
+		const targetBranch = this._getBranch(targetBranchName)
+
+		if (!targetBranch) {
+			return { success: false, message: `Unable to retrieve branch infomation for "${targetBranchName}". Unable to unlock, please contact Robomerge support."` }
+		}
+
+		const edge = this.getImmediateEdge(targetBranch)
+
+		if (!edge) {
+			return { success: false, message: `This bot has no edge for specified branch "${targetBranch.name}" on ${this.fullName}`}
+		}
+
+		edge.unblock("files unlocked")
+
+		return { success: true, message: 'Queued unlock request' }
+	}
+
 	createEdgeBlockageInfo(edgeBranch: Branch, failure: Failure, pending: PendingChange): BlockagePauseInfo {
 		// Check source node if it is currently tracking this conflict
 		let existingConflict = this.getConflictByBranchAndSourceCl(edgeBranch, pending.change.source_cl)
@@ -1007,7 +1127,8 @@ export class NodeBot extends PerforceStatefulBot implements NodeBotInterface {
 			failure.description,
 			<PendingChange> pending,
 			existingConflict ? existingConflict.time : new Date,
-			edgeBranch
+			edgeBranch,
+			failure.additionalInfo
 		)
 
 		// Check to see if this conflict has been acknowledged
@@ -1019,7 +1140,7 @@ export class NodeBot extends PerforceStatefulBot implements NodeBotInterface {
 		return pauseInfo
 	}
 
-	private static _makeBlockageInfo(errstr: string, errlong: string, arg: ChangeInfo | PendingChange, startedAt: Date, targetBranch?: Branch): BlockagePauseInfo {
+	private static _makeBlockageInfo(errstr: string, errlong: string, arg: ChangeInfo | PendingChange, startedAt: Date, targetBranch?: Branch, additionalInfo?: any): BlockagePauseInfo {
 		const MAX_LEN = 700
 		if (errlong.length > MAX_LEN) {
 			errlong = errlong.substring(0,MAX_LEN) + "..."
@@ -1034,7 +1155,8 @@ export class NodeBot extends PerforceStatefulBot implements NodeBotInterface {
 			change: info.cl,
 			sourceCl: info.source_cl,
 			source: info.source,
-			message: `${errstr}:\n${errlong}`
+			message: `${errstr}:\n${errlong}`,
+			additionalInfo
 		}
 
 		if (maybePendingChange.change) {
@@ -1216,6 +1338,16 @@ export class NodeBot extends PerforceStatefulBot implements NodeBotInterface {
 				)
 			}
 
+			if (blockage.failure.kind === 'Exclusive check-out') {
+				urls.unlockUrl = OperationUrlHelper.createUnlockUrl(
+					externalUrl,
+					botname,
+					branchname,
+					cl,
+					blockage.action.branch.name
+				)
+			}
+
 			// Do not generate skip urls if this branch is configure to disallow it
 			if (!disallowSkip) {
 				urls.skipUrl = OperationUrlHelper.createSkipUrl(
diff --git a/Engine/Extras/RoboMerge/v3/src/robo/notifications.ts b/Engine/Extras/RoboMerge/v3/src/robo/notifications.ts
--- a/Engine/Extras/RoboMerge/v3/src/robo/notifications.ts
+++ b/Engine/Extras/RoboMerge/v3/src/robo/notifications.ts
@@ -6,7 +6,7 @@ import { Args } from '../common/args';
 import { Badge } from '../common/badge';
 import { Random } from '../common/helper';
 import { ContextualLogger } from '../common/logger';
-import { Blockage, Branch, BranchArg, ForcedCl, MergeAction, NodeOpUrlGenerator, resolveBranchArg } from './branch-interfaces';
+import { Blockage, Branch, BranchArg, ExclusiveLockInfo, ForcedCl, MergeAction, NodeOpUrlGenerator, resolveBranchArg } from './branch-interfaces';
 import { PersistentConflict, Resolution } from './conflict-interfaces';
 import { BotEventHandler, BotEvents } from './events';
 import { NodeBot } from './nodebot';
@@ -545,7 +545,7 @@ export class BotNotifications implements BotEventHandler {
 			messagesToPost.push({ message })
 
 			if (blockage.failure.kind === 'Exclusive check-out') {
-				const exclusiveLockUsers = blockage.failure.additionalInfo as any[]			
+				const exclusiveLockUsers = (blockage.failure.additionalInfo as ExclusiveLockInfo).exclusiveLockUsers
 				let text = ''
 
 				for (const exclusiveLockUser of exclusiveLockUsers) {
diff --git a/Engine/Extras/RoboMerge/v3/src/robo/roboserver.ts b/Engine/Extras/RoboMerge/v3/src/robo/roboserver.ts
--- a/Engine/Extras/RoboMerge/v3/src/robo/roboserver.ts
+++ b/Engine/Extras/RoboMerge/v3/src/robo/roboserver.ts
@@ -632,6 +632,9 @@ class RoboWebApp implements AppInterface {
 			case "stomp":
 				return await readUtf8File('public/stomp.html')
 
+			case "unlock":
+				return await readUtf8File('public/unlock.html')
+
 			default:
 				return {statusCode: 404, message: `Unknown node operation requested: "${operation}"`}
 		}
@@ -763,6 +766,7 @@ export interface BlockageNodeOpUrls {
 	createShelfUrl?: string
 	skipUrl?: string
 	stompUrl?: string
+	unlockUrl?: string
 }
 export class OperationUrlHelper {
 	static createAcknowledgeUrl(externalRobomergeUrl: string, botname: string, branchname: string, changelistNum: string, edge?: string) {
@@ -801,4 +805,12 @@ export class OperationUrlHelper {
 			`&target=${encodeURIComponent(target)}`
 	}
 	
+	static createUnlockUrl(externalRobomergeUrl: string, botname: string, branchname: string, changelistNum: string, target: string) {
+		return `${externalRobomergeUrl}/op/unlock?` +
+			`bot=${encodeURIComponent(botname)}`+ 
+			`&branch=${encodeURIComponent(branchname)}` +
+			`&cl=${encodeURIComponent(changelistNum)}` +
+			`&target=${encodeURIComponent(target)}`
+	}
+
 }
diff --git a/Engine/Extras/RoboMerge/v3/src/robo/status-types.ts b/Engine/Extras/RoboMerge/v3/src/robo/status-types.ts
--- a/Engine/Extras/RoboMerge/v3/src/robo/status-types.ts
+++ b/Engine/Extras/RoboMerge/v3/src/robo/status-types.ts
@@ -70,6 +70,7 @@ export interface BlockagePauseInfo extends BlockagePauseInfoMinimal {
 	targetStream?: string  // Used by create shelf feature to filter user workspaces
 	sourceCl?: number
 	source?: string
+	additionalInfo?: any
 
 	// Set when the Acknowledge button is clicked
 	acknowledger?: string
diff --git a/Engine/Extras/RoboMerge/v3/src/robo/tick-journal.ts b/Engine/Extras/RoboMerge/v3/src/robo/tick-journal.ts
--- a/Engine/Extras/RoboMerge/v3/src/robo/tick-journal.ts
+++ b/Engine/Extras/RoboMerge/v3/src/robo/tick-journal.ts
@@ -10,6 +10,9 @@ export const ANALYTICS_SIMPLE_COUNTERS = {
 	// Keep track of how many successful stomp requests have been queued
 	stompQueued: 'stompreqqueued',
 
+	// Keep track of how many successful unlock requests have been queued
+	unlockQueued: 'unlockreqqueued',
+
 	// how many trace routes called (from CIS/build tools)
 	traces: 'traceroute',
 	failedTraces: 'tracefail'
