Immediately retry after forcing CL via webpage

diff --git a/Engine/Extras/RoboMerge/v3/src/robo/ipc.ts b/Engine/Extras/RoboMerge/v3/src/robo/ipc.ts
--- a/Engine/Extras/RoboMerge/v3/src/robo/ipc.ts
+++ b/Engine/Extras/RoboMerge/v3/src/robo/ipc.ts
@@ -586,7 +586,7 @@ export class IPC {
 				return {statusCode: 400, message: 'Invalid CL parameter: ' + cl}
 			}
 
-			let prevCl = generalOpTarget.forceSetLastClWithContext(cl, query.who, query.reason, !!query.unblock && query.unblock === 'true')
+			let prevCl = generalOpTarget.forceSetLastClWithContext(cl, query.who, query.reason, true)
 
 			this.ipcLogger.info(`Forcing last CL=${cl} on ${botname} : ${branch.name} (was CL ${prevCl}), ` +
 														`requested by ${query.who} (Reason: ${query.reason})`)
