Protect against robomergeUser being undefined

diff --git a/Engine/Extras/RoboMerge/v3/public/js/stomp.js b/Engine/Extras/RoboMerge/v3/public/js/stomp.js
--- a/Engine/Extras/RoboMerge/v3/public/js/stomp.js
+++ b/Engine/Extras/RoboMerge/v3/public/js/stomp.js
@@ -200,7 +200,7 @@ function transitionDisplayStompResults() {
 // Create visualization of the changelist, with links to swarm, a display of the description and a list of relevant files to this request
 function visualizeChangelist(changelist, changelistData) {
     const changelistDiv = $(`<div id="${changelist}" class="stomp-visual">`)
-    if (changelistData.author.toLowerCase() === robomergeUser.userName.toLowerCase()) {
+    if (robomergeUser && changelistData.author.toLowerCase() === robomergeUser.userName.toLowerCase()) {
         changelistDiv.append($('<h3>').html(`${makeClLink(changelist)} by <strong>you</strong>`))
     } else {
         changelistDiv.append($('<h3>').html(`${makeClLink(changelist)} by ${changelistData.author}`))
diff --git a/Engine/Extras/RoboMerge/v3/public/js/unlock.js b/Engine/Extras/RoboMerge/v3/public/js/unlock.js
--- a/Engine/Extras/RoboMerge/v3/public/js/unlock.js
+++ b/Engine/Extras/RoboMerge/v3/public/js/unlock.js
@@ -200,7 +200,7 @@ function transitionDisplayUnlockResults() {
 // Create visualization of the file, with links to swarm
 function visualizeAuthor(author, files) {
     const authorDiv = $(`<div id="${author}" class="unlock-visual">`)
-    if (author.toLowerCase() === robomergeUser.userName.toLowerCase()) {
+    if (robomergeUser && author.toLowerCase() === robomergeUser.userName.toLowerCase()) {
         authorDiv.append($('<h3>').html("Locked by <strong>you</strong>"))
     } else {
         authorDiv.append($('<h3>').html(`Locked by ${author}`))
