Always analyze all errors and to identify all exclusively locked files as the unlock files feature needs the full list

diff --git a/Engine/Extras/RoboMerge/v3/src/robo/branch-interfaces.ts b/Engine/Extras/RoboMerge/v3/src/robo/branch-interfaces.ts
--- a/Engine/Extras/RoboMerge/v3/src/robo/branch-interfaces.ts
+++ b/Engine/Extras/RoboMerge/v3/src/robo/branch-interfaces.ts
@@ -168,7 +168,6 @@ export interface ChangeInfo extends TargetInfo {
 		address: string
 	}
 	targetWorkspaceOverride?: string
-	analyzeAllIntegrationErrors?: boolean
 	overriddenCommand: string
 	macros: string[]
 }
diff --git a/Engine/Extras/RoboMerge/v3/src/robo/edgebot.ts b/Engine/Extras/RoboMerge/v3/src/robo/edgebot.ts
--- a/Engine/Extras/RoboMerge/v3/src/robo/edgebot.ts
+++ b/Engine/Extras/RoboMerge/v3/src/robo/edgebot.ts
@@ -259,13 +259,10 @@ class EdgeBotImpl extends PerforceStatefulBot {
 		return results
 	}
 
-	private async analyzeIntegrationError(errors: string[], analyzeAllIntegrationErrors?: boolean) {
-		if (!analyzeAllIntegrationErrors && errors.length > MAX_INTEGRATION_ERRORS_TO_ANALYZE) {
-			this.edgeBotLogger.error(`Integration error: ${errors.length} files, checking first ${MAX_INTEGRATION_ERRORS_TO_ANALYZE}`)
-		}
+	private async analyzeIntegrationError(errors: string[]) {
 		
 		const openedRequests: [RegExpMatchArray, Promise<OpenedFileRecord[]>, Promise<OpenedFileRecord[]>][] = []
-		for (const err of analyzeAllIntegrationErrors ? errors : errors.slice(0, MAX_INTEGRATION_ERRORS_TO_ANALYZE)) {
+		for (const err of errors) {
 			const match = err.match(EXCLUSIVE_CHECKOUT_REGEX)
 			if (match) {
 				openedRequests.push([match, this.p4.opened(null, match[1] + match[2], true), this.p4.opened(null, match[1] + match[2])])
@@ -554,7 +551,7 @@ class EdgeBotImpl extends PerforceStatefulBot {
 		}
 
 		const errors = results as string[]
-		const exclusiveFiles = await this.analyzeIntegrationError(errors, info.analyzeAllIntegrationErrors)
+		const exclusiveFiles = await this.analyzeIntegrationError(errors)
 
 		const description = errors.join('\n')
 		let failure: Failure | null = null
@@ -562,8 +559,8 @@ class EdgeBotImpl extends PerforceStatefulBot {
 		if (exclusiveFiles.length > 0) {
 			// will need to store the exclusive file if we want to @ people in Slack
 			const exclCheckoutMessages = exclusiveFiles.map(exc => `${exc.depotPath} checked out by ${exc.user}`)
-			if (!info.analyzeAllIntegrationErrors && errors.length > MAX_INTEGRATION_ERRORS_TO_ANALYZE) {
-				exclCheckoutMessages.push(`... and up to ${errors.length - MAX_INTEGRATION_ERRORS_TO_ANALYZE} more`)
+			if (errors.length > MAX_INTEGRATION_ERRORS_TO_ANALYZE) {
+				exclCheckoutMessages.push(`... and ${errors.length - MAX_INTEGRATION_ERRORS_TO_ANALYZE} more`)
 			}
 			const exclusiveLockUsers = Array.from(new Set(exclusiveFiles.map(exc => `${exc.user.toLowerCase()}`))).map(user => ({user, userEmail: this.p4.getEmail(user)}))
 			failure = { kind: 'Exclusive check-out', description, summary: exclCheckoutMessages.join('\n'), additionalInfo: {exclusiveLockUsers,exclusiveFiles} }
