Change how the unlock and stomp intermediate verification step being successful is visualized to try and reduce indications to user that the operation is complete and they do not need to do anything else

diff --git a/Engine/Extras/RoboMerge/v3/public/js/stomp.js b/Engine/Extras/RoboMerge/v3/public/js/stomp.js
--- a/Engine/Extras/RoboMerge/v3/public/js/stomp.js
+++ b/Engine/Extras/RoboMerge/v3/public/js/stomp.js
@@ -22,7 +22,7 @@ function stompVerify() {
     let stompOperation = $.Deferred()
     
     // Get requested branch infomation and setup the page data
-    let visualizationOperation = getBranch(requestedBotName, requestedBranchName, function(data) {
+    getBranch(requestedBotName, requestedBranchName, function(data) {
         try {
             // Ensure we have data
             if (!data) {
@@ -98,8 +98,10 @@ function stompVerify() {
                 visualizeStompVerification(requestedBranchCl, stompJson)
 
                 if (stompJson.validRequest) {
-                    $('#afterVerificationResultText').html('<span><i class="fas fa-check-circle"></i></span> Stomp Verification Success!')
-                    $('#afterVerificationResultText').css('color', 'green');
+                    $('#afterVerificationResultText').html('Stomp Verification Complete')
+                    $('#afterVerificationResultText').append(
+                        $('<h4 style="text-align: center">').text("Use the button below to proceed with the stomp operation")
+                    )
                 } else {
                     $('#afterVerificationResultText').html('<span><i class="fas fa-exclamation-triangle"></i></span> Stomp Verification Returned Issues.')
                     $('#afterVerificationResultText').css('color', 'red');
@@ -111,9 +113,6 @@ function stompVerify() {
                 if (stompJson.validRequest) {
                     const formButtonDiv = $('#formButtons')
 
-                    formButtonDiv.append(
-                        $('<h3 style="text-align: center">').text("Use the button below to proceed with the stomp operation:")
-                    )
                     // Return to Robomerge homepage
                     let cancelButton = $('<button type="button" class="btn btn-lg btn-info" style="margin:.3em;">').text(`Cancel`).appendTo(formButtonDiv)
                     cancelButton.click(function() {
diff --git a/Engine/Extras/RoboMerge/v3/public/js/unlock.js b/Engine/Extras/RoboMerge/v3/public/js/unlock.js
--- a/Engine/Extras/RoboMerge/v3/public/js/unlock.js
+++ b/Engine/Extras/RoboMerge/v3/public/js/unlock.js
@@ -22,7 +22,7 @@ function unlockVerify() {
     let unlockOperation = $.Deferred()
     
     // Get requested branch infomation and setup the page data
-    let visualizationOperation = getBranch(requestedBotName, requestedBranchName, function(data) {
+    getBranch(requestedBotName, requestedBranchName, function(data) {
         try {
             // Ensure we have data
             if (!data) {
@@ -98,8 +98,11 @@ function unlockVerify() {
                 visualizeUnlockVerification(requestedBranchCl, unlockJson)
 
                 if (unlockJson.validRequest) {
-                    $('#afterVerificationResultText').html('<span><i class="fas fa-check-circle"></i></span> Unlock Verification Success!')
-                    $('#afterVerificationResultText').css('color', 'green');
+                    $('#afterVerificationResultText').html('Unlock Verification Complete')
+                    $('#afterVerificationResultText').append(
+                        $('<h4 style="text-align: center">').text("Use the button below to proceed with the unlock operation")
+                    )
+
                 } else {
                     $('#afterVerificationResultText').html('<span><i class="fas fa-exclamation-triangle"></i></span> Unlock Verification Returned Issues.')
                     $('#afterVerificationResultText').css('color', 'red');
@@ -111,9 +114,6 @@ function unlockVerify() {
                 if (unlockJson.validRequest) {
                     const formButtonDiv = $('#formButtons')
 
-                    formButtonDiv.append(
-                        $('<h3 style="text-align: center">').text("Use the button below to proceed with the unlock operation:")
-                    )
                     // Return to Robomerge homepage
                     let cancelButton = $('<button type="button" class="btn btn-lg btn-info" style="margin:.3em;">').text(`Cancel`).appendTo(formButtonDiv)
                     cancelButton.click(function() {
