Merging using ROBO://GamePlugins/Main/Programs/robomerge-ts->//UE5/Main/Engine/Extras/RoboMerge/v3

Add Okta sign-in option for Robomerge

#rnx

diff --git a/Engine/Extras/RoboMerge/v3/.vscode/launch.json b/Engine/Extras/RoboMerge/v3/.vscode/launch.json
--- a/Engine/Extras/RoboMerge/v3/.vscode/launch.json
+++ b/Engine/Extras/RoboMerge/v3/.vscode/launch.json
@@ -29,7 +29,11 @@
                 "EPIC_ENV": "dev",
                 "EPIC_DEPLOYMENT": "local",
                 "BOTNAME": "test",
-                "ROBO_RUNBOTS": "yes"
+                "ROBO_RUNBOTS": "yes",
+                "ROBO_PERSISTENCE_BACKUP_FREQUENCY": "0",
+                "ROBO_USE_AUTH_IN_DEV": "false",
+                "OKTA_SIGN_IN": "false",
+                "LDAP_SIGN_IN": "true"
             },
             "smartStep": true,
             "skipFiles": [
diff --git a/Engine/Extras/RoboMerge/v3/config/ldap.cfg.json b/Engine/Extras/RoboMerge/v3/config/ldap.cfg.json
--- a/Engine/Extras/RoboMerge/v3/config/ldap.cfg.json
+++ b/Engine/Extras/RoboMerge/v3/config/ldap.cfg.json
@@ -15,6 +15,10 @@
 		"connectTimeout": 30000
 	},
 	"bot-groups": [
+		{
+			"group": "Robomerge_Admin_Access",
+			"tags": ["admin"]
+		},
 		{
 			"group": "Robomerge_FTE_Access",
 			"tags": ["fte"]
diff --git a/Engine/Extras/RoboMerge/v3/config/okta.cfg.json b/Engine/Extras/RoboMerge/v3/config/okta.cfg.json
new file mode 100644
--- /dev/null
+++ b/Engine/Extras/RoboMerge/v3/config/okta.cfg.json
@@ -0,0 +1,6 @@
+{
+    "clientId": "<okta clientId>",
+    "redirectUri": "/login",
+    "issuer": "https://yourcompany.okta.com/oauth2/authorizationServerURL",
+    "scopes": ["openid", "profile", "email"]
+}
\ No newline at end of file
diff --git a/Engine/Extras/RoboMerge/v3/public/js/boilerplate.js b/Engine/Extras/RoboMerge/v3/public/js/boilerplate.js
--- a/Engine/Extras/RoboMerge/v3/public/js/boilerplate.js
+++ b/Engine/Extras/RoboMerge/v3/public/js/boilerplate.js
@@ -385,6 +385,7 @@ function generateRobomergeHeader(createSignedInUserDiv = true) {
 		let logOutButton = $('<button id="log-out" class="btn btn-xs btn-warning">Sign out</button>')
 		logOutButton.click(function() {
 			document.cookie = 'auth=; path=; redirect_to=;';
+			document.cookie = 'signedOut=true;';
 			window.location.href = '/login';
 		})
 		loggedInUser.append(logOutButton)
diff --git a/Engine/Extras/RoboMerge/v3/public/js/login.js b/Engine/Extras/RoboMerge/v3/public/js/login.js
--- a/Engine/Extras/RoboMerge/v3/public/js/login.js
+++ b/Engine/Extras/RoboMerge/v3/public/js/login.js
@@ -12,17 +12,97 @@ function getCookie(name) {
 	return null;
 }
 
-$(() => {
+var authClient;
 
-	const $form = $('.login-form form');
+$(async () => {
+
+	const $ldapForm = $('.login-form .ldap-form');
 
 	const clearError = () => $('#login-error').hide();
 
-	const $nameInput = $('input[name=user]', $form).on('input', clearError);
-	const $passwordInput = $('input[name=password]', $form).on('input', clearError);
+	const $nameInput = $('input[name=user]', $ldapForm).on('input', clearError);
+	const $passwordInput = $('input[name=password]', $ldapForm).on('input', clearError);
+
+	const $oktaForm = $('.login-form .okta-form');
+	const $oktaError = $('#okta-login-error');
+
+	const signInMethod = await $.get('/signInMethod');
+
+	if (signInMethod.okta) {
+		$oktaForm.show();
+
+		try {
+			const config = await $.get('/oktaConfig');
+			authClient = new OktaAuth(JSON.parse(config));
+		}
+		catch(error) {
+			console.error(error);
+			$('html').show();
+			$oktaError.text("OktaAuth is not configured").show();
+		}
+
+		// If this is a login redirect from Okta we can parse the tokens returned from Okta
+		if (authClient && authClient.isLoginRedirect()) {
+
+			try {
+				const { tokens } = await authClient.token.parseFromUrl();
+				authClient.tokenManager.setTokens(tokens);
+
+				const { accessToken, idToken } = await authClient.tokenManager.getTokens();
+				const user = await authClient.token.getUserInfo(accessToken, idToken);
+
+				const postData = {
+					user: user.email,
+					displayName: user.preferred_username,
+					groups: JSON.stringify(user.groups)
+				};
+
+				$.post('/oktaLogin', postData)
+				.done((token) => {
+					// Check to see if we have a redirect request from the website
+					var urlParams = new URLSearchParams(window.location.search);
+					var redirectString = "/"
+					if (urlParams.has("redirect")) {
+						redirectString = decodeURIComponent(urlParams.get("redirect"))
+					}
+
+					document.cookie = `auth=${token}; secure=true; path=/;`;
+					document.cookie = 'redirect_to=;';
+					window.location = window.origin + redirectString;
+				})
+				.fail((xhr, status, error) => {
+					$('html').show();
+					$oktaError.text(error + ": " + xhr.responseText).show();
+				});
+
+			}
+			catch(error) {
+				$('html').show();
+				$oktaError.text(error).show();
+			}
+
+		}
+		else {
+			// Check to see if user just signed out
+			if (getCookie('signedOut')) {
+				$('html').show();
+				document.cookie = 'signedOut=;';
+			}
+			else if (authClient) {
+				oktaSignIn();
+			}
+		}
+	}
 
-	$form.on('submit', () => {
-		const $submitButton = $('button', $form).prop('disabled', true).text('Signing in ...');
+	if (signInMethod.ldap) {
+		if (!signInMethod.okta) {
+			$('html').show();
+		}
+		$ldapForm.show();
+	}
+
+	$ldapForm.on('submit', () => {
+		const $submitButton = $('button', $ldapForm).prop('disabled', true).text('Signing in ...');
 
 		// clear auth on attempt to log in
 		document.cookie = 'auth=;';
@@ -62,4 +142,16 @@ $(() => {
 
 		return false;
 	});
+
+	
 });
+
+async function oktaSignIn() {
+	// Clear auth cookie on attempt to sign in
+	document.cookie = 'auth=;';
+
+	// Start full-page redirect to Okta
+	if (authClient) {
+		authClient.token.getWithRedirect();
+	}
+}
diff --git a/Engine/Extras/RoboMerge/v3/public/login.html b/Engine/Extras/RoboMerge/v3/public/login.html
--- a/Engine/Extras/RoboMerge/v3/public/login.html
+++ b/Engine/Extras/RoboMerge/v3/public/login.html
@@ -1,5 +1,5 @@
 <!DOCTYPE html>
-<html>
+<html style="display: none">
 <head>
 <title>ROBOMERGE - sign in</title>
 
@@ -21,11 +21,15 @@
 
 <script src="/js/boilerplate.js"></script>
 <script src="/js/login.js"></script>
+
+<script src="https://global.oktacdn.com/okta-auth-js/7.4.3/okta-auth-js.polyfill.js" type="text/javascript"></script>
+<script src="https://global.oktacdn.com/okta-auth-js/7.4.3/okta-auth-js.min.js" type="text/javascript"></script>
+
 </head>
 <body>
 
 <div class="login-form">
-  <form class="">
+  <form class="ldap-form" style="display: none">
     <div id="login-error" class="alert alert-danger initiallyHidden">Invalid username or password!</div>
 
     <div class="text-center"><img src="/img/epic.png"><br><br>Please sign in with your Epic credentials.</div>
@@ -37,9 +41,13 @@
       <input type="password" name="password" class="form-control" placeholder="Password" required="required">
     </div>
     <div class="form-group">
-      <button type="submit" class="btn btn-primary btn-block">Sign in</button>
+      <button type="submit" class="btn btn-primary btn-block">Sign in with LDAP</button>
     </div>
   </form>
+  <form class="okta-form" style="display: none;">
+    <div id="okta-login-error" class="alert alert-danger initiallyHidden"></div>
+    <button class="btn btn-primary btn-block" onclick="oktaSignIn()">Sign in with Okta</button>
+  </form>
 </div>
 
 </body>
diff --git a/Engine/Extras/RoboMerge/v3/src/robo/roboserver.ts b/Engine/Extras/RoboMerge/v3/src/robo/roboserver.ts
--- a/Engine/Extras/RoboMerge/v3/src/robo/roboserver.ts
+++ b/Engine/Extras/RoboMerge/v3/src/robo/roboserver.ts
@@ -40,6 +40,24 @@ let ENVIRONMENT: {[param: string]: any}
 			parse: str => str === "false" ? false : true,
 			env: 'ROBO_PREVIEW_ONLY',
 			dflt: false
+		},
+		useAuthInDev: {
+			match: /^(-useAuthInDev)$/,
+			parse: str => str === "false" ? false : true,
+			env: 'ROBO_USE_AUTH_IN_DEV',
+			dflt: false
+		},
+		oktaSignIn: {
+			match: /^(-oktaSignIn)$/,
+			parse: str => str === "false" ? false : true,
+			env: 'OKTA_SIGN_IN',
+			dflt: false
+		},
+		ldapSignIn: {
+			match: /^(-ldapSignIn)$/,
+			parse: str => str === "false" ? false : true,
+			env: 'LDAP_SIGN_IN',
+			dflt: true
 		}
 	}
 
@@ -111,7 +129,13 @@ export class RoboServer {
 		}
 
 		//this.server.addFileMapping('/', 'index.html')
-		this.server.addFileMapping('/login', 'login.html', {secureOnly: true})
+		if (ENVIRONMENT.devMode && ENVIRONMENT.useAuthInDev) {
+			this.server.addFileMapping('/login', 'login.html', {secureOnly: false})
+		}
+		else {
+			this.server.addFileMapping('/login', 'login.html', {secureOnly: true})
+		}
+
 		this.server.addFileMapping('/allbots', 'allbots.html')
 		this.server.addFileMapping('/js/*.wasm', 'bin/$1.wasm.gz', {
 			filetype: "application/wasm", 
@@ -368,6 +392,34 @@ class RoboWebApp implements AppInterface {
 		return token || {statusCode: 401, message: 'invalid credentials'};
 	}
 
+	@Handler('POST', '/oktaLogin')
+	async oktaLogin() {
+		if (!this.request.reqData) {
+			return {statusCode: 400, message: 'no log-in data received'}
+		}
+		
+		const creds = querystring.parse(this.request.reqData);
+		if (!creds.user || Array.isArray(creds.user) || !creds.displayName || Array.isArray(creds.displayName) || !creds.groups || Array.isArray(creds.groups)) {
+			return {statusCode: 400, message: 'invalid log-in data'}
+		}
+		const token = Session.oktaLogin(creds.user, creds.displayName, creds.groups);
+		return token || {statusCode: 401, message: 'invalid credentials'};
+	}
+
+	@Handler('GET', '/oktaConfig')
+	async oktaConfig() {
+		return fs.readFileSync('config/okta.cfg.json', 'utf8');
+	}
+
+	@Handler('GET', '/signInMethod')
+	async signInMethod() {
+		const signIn = {
+			okta: ENVIRONMENT.oktaSignIn,
+			ldap: ENVIRONMENT.ldapSignIn
+		};
+		return signIn;
+	}
+
 	@SecureHandler('PUT', '/api/control/verbose/*')
 	async setVerbose(onOff: string): Promise<OperationReturnType> {
 		if (!this.authData) {
@@ -659,11 +711,30 @@ class RoboWebApp implements AppInterface {
 			}
 		}
 		else if (ENVIRONMENT.devMode) {
-			const user = ENVIRONMENT.devModeUser || 'dev'
-			this.authData = {
-				user,
-				displayName: `${user} (dev mode)`,
-				tags: new Set(['fte', 'admin'])
+			if (ENVIRONMENT.useAuthInDev) {
+				const authToken = getCookie(this.getCookies(), 'auth')
+				if (authToken) {
+					this.authData = Session.tokenToAuthData(authToken)
+				}
+
+				if (!this.authData) {
+					let redirectString = this.request.url.pathname
+					if (this.request.url.search) {
+						redirectString = `${redirectString}${this.request.url.search}`
+					}
+
+					return {statusCode: 302, message: 'Must log in', headers: [
+						['Location', `/login?redirect=${encodeURIComponent(redirectString)}`]
+					]}
+				}
+			}
+			else {
+				const user = ENVIRONMENT.devModeUser || 'dev'
+				this.authData = {
+					user,
+					displayName: `${user} (dev mode)`,
+					tags: new Set(['fte', 'admin'])
+				}
 			}
 		}
 		else {
diff --git a/Engine/Extras/RoboMerge/v3/src/robo/session.ts b/Engine/Extras/RoboMerge/v3/src/robo/session.ts
--- a/Engine/Extras/RoboMerge/v3/src/robo/session.ts
+++ b/Engine/Extras/RoboMerge/v3/src/robo/session.ts
@@ -154,6 +154,21 @@ export class Session {
 		})
 	}
 
+	static oktaLogin(user: string, displayName: string, groups: string) {
+		let groupsArray: string[] = JSON.parse(groups);
+		let tags = new Set<string>();
+
+		for (const groupInfo of groupsArray) {
+			const tagsForGroup = Session.BOT_GROUPS.get(groupInfo);
+			if (tagsForGroup) {
+				tags = new Set<string>([...tagsForGroup, ...tags]);
+			}
+		}
+		
+		const authData = {user: user, displayName: displayName, tags};
+		return Session.authDataToToken(authData);
+	}
+
 	static tokenToAuthData(token: string): AuthData | null {
 		const parts = token.replace(/-/g, '+').replace(/_/g, '/').split(':')
 		if (parts.length !== 2) {
