Make user and displayName in Okta consistent with LDAP

diff --git a/Engine/Extras/RoboMerge/v3/public/js/login.js b/Engine/Extras/RoboMerge/v3/public/js/login.js
--- a/Engine/Extras/RoboMerge/v3/public/js/login.js
+++ b/Engine/Extras/RoboMerge/v3/public/js/login.js
@@ -52,8 +52,8 @@ $(async () => {
 				const user = await authClient.token.getUserInfo(accessToken, idToken);
 
 				const postData = {
-					user: user.email,
-					displayName: user.preferred_username,
+					user: user.preferred_username,
+					displayName: user.name,
 					groups: JSON.stringify(user.groups)
 				};
 
