catch invalid usernames and notify alerts channel instead of crashing

diff --git a/Engine/Extras/RoboMerge/v3/src/common/perforce.ts b/Engine/Extras/RoboMerge/v3/src/common/perforce.ts
--- a/Engine/Extras/RoboMerge/v3/src/common/perforce.ts
+++ b/Engine/Extras/RoboMerge/v3/src/common/perforce.ts
@@ -611,21 +611,40 @@ export class PerforceContext {
 			opts.edgeServerAddress = edgeServerAddress
 		}
 
-		let parsedLoadedClients = this._execP4Ztag(null, args, opts);
-		let parsedUnloadedClients = (includeUnloaded ? this._execP4Ztag(null, [...args, '-U'], opts) : null)
 		let workspaces = [];
-		for (let clientDef of await parsedLoadedClients) {
-			if (clientDef.client) {
-				workspaces.push(clientDef);
-			}
-		}
-		if (includeUnloaded) {
-			for (let clientDef of await parsedUnloadedClients!) {
+		try {
+			let parsedLoadedClients = this._execP4Ztag(null, args, opts);
+			let parsedUnloadedClients = (includeUnloaded ? this._execP4Ztag(null, [...args, '-U'], opts) : null)
+			for (let clientDef of await parsedLoadedClients) {
 				if (clientDef.client) {
 					workspaces.push(clientDef);
 				}
 			}
+			if (includeUnloaded) {
+				for (let clientDef of await parsedUnloadedClients!) {
+					if (clientDef.client) {
+						workspaces.push(clientDef);
+					}
+				}
+			}
 		}
+		catch (reason) {
+			if (!isExecP4Error(reason)) {
+				throw reason
+			}
+
+			let [err, output] = reason
+
+			// if this change has already been integrated, this is a special return (still a success)
+			if (!output.includes("Revision chars (@, #) not allowed in")) {
+				throw err
+			}
+
+			const errorMsg = `Attempted to find workspaces for invalid user ${user || this.username}}`
+			this.logger.error(errorMsg)
+			postToRobomergeAlerts(errorMsg)
+		}
+
 		return workspaces as ClientSpec[];
 	}
 
