Remove excess }

diff --git a/Engine/Extras/RoboMerge/v3/src/common/perforce.ts b/Engine/Extras/RoboMerge/v3/src/common/perforce.ts
--- a/Engine/Extras/RoboMerge/v3/src/common/perforce.ts
+++ b/Engine/Extras/RoboMerge/v3/src/common/perforce.ts
@@ -640,7 +640,7 @@ export class PerforceContext {
 				throw err
 			}
 
-			const errorMsg = `Attempted to find workspaces for invalid user ${user || this.username}}`
+			const errorMsg = `Attempted to find workspaces for invalid user ${user || this.username}`
 			this.logger.error(errorMsg)
 			postToRobomergeAlerts(errorMsg)
 		}
