Catch case where the files to unlock are already unlocked

diff --git a/Engine/Extras/RoboMerge/v3/src/common/perforce.ts b/Engine/Extras/RoboMerge/v3/src/common/perforce.ts
--- a/Engine/Extras/RoboMerge/v3/src/common/perforce.ts
+++ b/Engine/Extras/RoboMerge/v3/src/common/perforce.ts
@@ -1183,7 +1183,20 @@ export class PerforceContext {
 	async revertFiles(files: string[], client: string) {
 		const edgeServer = await this.getWorkspaceEdgeServer(client)
 		const args = ['revert', '-C', client, ...files]
-		return this._execP4Ztag(null, args, {edgeServerAddress: edgeServer?.address})
+		try {
+			await this._execP4Ztag(null, args, {edgeServerAddress: edgeServer?.address})
+		}
+		catch (reason) {
+			if (!isExecP4Error(reason)) {
+				throw reason
+			}
+
+			let [err, output] = reason
+			// this happens if there's literally nothing in the CL. consider this a success
+			if (!output.match(/file\(s\) not opened (?:on this client|in that changelist)\./)) {
+				throw err;
+			}
+		}
 	}
 
 	// revert a CL deleting any files marked for add
