Don't send slack messages in dev when useSlackInDev is false even if the bot/user tokens are in the vault

diff --git a/Engine/Extras/RoboMerge/v3/src/robo/notifications.ts b/Engine/Extras/RoboMerge/v3/src/robo/notifications.ts
--- a/Engine/Extras/RoboMerge/v3/src/robo/notifications.ts
+++ b/Engine/Extras/RoboMerge/v3/src/robo/notifications.ts
@@ -994,8 +994,8 @@ export function bindBotNotifications(events: BotEvents, slackChannelOverrides: [
 
 	let slackMessages
 
-	const botToken = args.devMode && !args.useSlackInDev && SLACK_DEV_DUMMY_TOKEN || SLACK_TOKENS.bot
-	const userToken = args.devMode && !args.useSlackInDev && SLACK_DEV_DUMMY_TOKEN || SLACK_TOKENS.user
+	const botToken = (!args.devMode || args.useSlackInDev) && SLACK_TOKENS.bot || args.devMode && args.useSlackInDev && SLACK_DEV_DUMMY_TOKEN 
+	const userToken = (!args.devMode || args.useSlackInDev) && SLACK_TOKENS.user || args.devMode && args.useSlackInDev && SLACK_DEV_DUMMY_TOKEN
 	if (botToken && events.botConfig.slackChannel) {
 		logger.info('Enabling Slack messages for ' +  events.botname)
 		slackMessages = new SlackMessages(new Slack({id: events.botConfig.slackChannel, botToken, userToken}, args.slackDomain), persistence, logger)
