TrackChanges will now link to the related github commit for streams/changes that are mirrored

diff --git a/Engine/Extras/RoboMerge/v3/public/js/trackchange.js b/Engine/Extras/RoboMerge/v3/public/js/trackchange.js
--- a/Engine/Extras/RoboMerge/v3/public/js/trackchange.js
+++ b/Engine/Extras/RoboMerge/v3/public/js/trackchange.js
@@ -40,15 +40,45 @@ function getStreamGraphName(streamDisplayName) {
 	return streamDisplayName.split('/').reverse()[0]
 }
 
-function generateChangeList(data) {
+async function getGitHubCommit(cl) {
+	const timeout = new Promise((resolve, reject) => {
+    	setTimeout(() => reject(new Error('GitSync request took longer than 3 seconds')), 3000)
+    });
+     const response = await Promise.race([
+		$.get(`https://gitsync.devtools.epicgames.com/api/v1/changes/${cl}`),
+        timeout
+    ]);
+    return {"commit":response.commit, "commitURL":response.commitURL};
+}
+
+async function getGitHubCommits(changes) {
+	return new Map(await Promise.all(Object.keys(changes).map(async cl => {
+		try {
+            const ghCommit = await getGitHubCommit(cl);
+            return [cl, ghCommit];
+        } catch(error) {
+            return [cl, null];
+        }
+	})));
+}
+
+function formatGitHubCommit(change) {
+	if (change) {
+		return `<td><a href="${change.commitURL}" target="_blank"><button class="btn btn-sm">${change.commit.slice(0,6)}</button></a></td>`
+	}
+	return '<td></td>'
+}
+
+async function generateChangeList(dataObj) {
 
-	let dataObj = data.data
+	const gitHubCommits = await getGitHubCommits(dataObj.changes)
 	let html = '<div style="margin: auto; width: 80%;"><table class="table"><tbody>'
 	for (const cl in dataObj.changes) {
 		console.log()
-		html += '<tr>'
+		html += '<tr valign="middle">'
 		html += `<td><b>${dataObj.changes[cl].streamDisplayName}</b></td>`
 		html += `<td><a href="https://p4-swarm.epicgames.net/changes/${cl}" target="_blank">CL#${cl}</a></td>`
+		html += `${formatGitHubCommit(gitHubCommits.get(cl))}`
 		//html += `<td style="text-align:center;">${getMergeMethodString(dataObj.changes[cl].mergeMethod)}</td>`
 		//html += `<td style="text-align:center;">${dataObj.changes[cl].sourceCL}</td>`
 		html += '</tr>'
@@ -109,10 +139,10 @@ function createGraph(changes) {
     return graphContainer;	
 }
 
-function buildResults(data) {
+async function buildResults(data) {
 	if (Object.keys(data.data.changes).length > 0) {
 		const $successPanel = $('#success-panel');
-		$('#changes', $successPanel).html(generateChangeList(data));
+		$('#changes', $successPanel).html(await generateChangeList(data.data));
 		$successPanel.show();
 		$('#graph').append(createGraph(data.data.changes))
 	} else {
