Correctly mark functional tests as using slack in dev

diff --git a/Engine/Extras/RoboMerge/v3/run-tests.sh b/Engine/Extras/RoboMerge/v3/run-tests.sh
--- a/Engine/Extras/RoboMerge/v3/run-tests.sh
+++ b/Engine/Extras/RoboMerge/v3/run-tests.sh
@@ -76,6 +76,7 @@ docker run -d -p 8877:8877 -p 8811:8811 -p 25:25 -h robomerge_functtest --name r
     -e P4PORT=p4docker:1666  \
     -e ROBO_BRANCHSPECS_ROOT_PATH=//RoboMergeData/Main  \
     -e ROBO_PERSISTENCE_BACKUP_FREQUENCY=0 \
+    -e ROBO_USE_SLACK_IN_DEV=true \
     -e ROBO_NO_MAIL=true \
     -e ROBO_NO_TLS=true  \
     -e ROBO_EXTERNAL_URL=http://robomerge_functtest:8877  \
