Fix use of #robomerge #manual failing crashing when the desired client is on an edgeserver in most cases
Fix #robomerge #manual not sending a notification message that the shelf is ready.

diff --git a/Engine/Extras/RoboMerge/v3/src/robo/nodebot.ts b/Engine/Extras/RoboMerge/v3/src/robo/nodebot.ts
--- a/Engine/Extras/RoboMerge/v3/src/robo/nodebot.ts
+++ b/Engine/Extras/RoboMerge/v3/src/robo/nodebot.ts
@@ -1446,10 +1446,10 @@ export class NodeBot extends PerforceStatefulBot implements NodeBotInterface {
 			return
 		}
 
-		const owner = pendingChange.change.owner
+		const owner = pendingChange.change.owner || pendingChange.change.author
 		// Should never happen
 		if (!owner) {
-			this.nodeBotLogger.warn(`Unable to send shelf creation notification for source CL ${pendingChange.change.source_cl} in branch ${shelfBranch} because there was no owner.`)
+			this.nodeBotLogger.warn(`Unable to send shelf creation notification for source CL ${pendingChange.change.source_cl} in branch ${shelfBranch.name} because there was no owner.`)
 			return
 		}
 
@@ -1580,10 +1580,13 @@ export class NodeBot extends PerforceStatefulBot implements NodeBotInterface {
 			return emResult
 		}
 
-		// should also do this for #manual changes (set up some testing around those first)
-		if (change.forceCreateAShelf || (change.isUserRequest && !change.forceStompChanges)) {
+		// this will deal with #manual changes with a single target and will also manage to work for multiple
+		// targets as long as all the target workspaces are on the same edge, but it will still crash if you
+		// have multiple targets with workspaces on different edgeservers
+		let isManualChange = result.info.targets ? result.info.targets![0].flags.has('manual') : false;
+		if (change.forceCreateAShelf || (change.isUserRequest && !change.forceStompChanges) || isManualChange) {
 			if (!optWorkspaceOverride && (result.info.targets || []).length == 1) {
-				result.info.targetWorkspaceForShelf = await p4util.chooseBestWorkspaceForUser(this.p4, result.info.owner!, result.info.targets![0].branch.stream)
+				result.info.targetWorkspaceForShelf = await p4util.chooseBestWorkspaceForUser(this.p4, result.info.owner||result.info.author, result.info.targets![0].branch.stream)
 				optWorkspaceOverride = result.info.targetWorkspaceForShelf
 				this.nodeBotLogger.info(`Chose workspace ${optWorkspaceOverride}`)
 			}
