Log slack errors instead of throwing which is generally uncaught and crashes the application.
Remove unused function

diff --git a/Engine/Extras/RoboMerge/v3/src/robo/notifications.ts b/Engine/Extras/RoboMerge/v3/src/robo/notifications.ts
--- a/Engine/Extras/RoboMerge/v3/src/robo/notifications.ts
+++ b/Engine/Extras/RoboMerge/v3/src/robo/notifications.ts
@@ -78,7 +78,7 @@ export async function postToRobomergeAlerts(message: string) {
 
 export async function postMessageToChannel(message: string, channel: string, style: SlackMessageStyles = SlackMessageStyles.GOOD) {
 	if (SLACK_TOKENS.bot) {
-		return new Slack({id: channel, botToken: SLACK_TOKENS.bot, userToken: SLACK_TOKENS.user}, args.slackDomain).postMessageToDefaultChannel({
+		return new Slack({id: channel, botToken: SLACK_TOKENS.bot, userToken: SLACK_TOKENS.user}, args.slackDomain,new ContextualLogger(`Post Message to ${channel}`)).postMessageToDefaultChannel({
 			text: message,
 			style,
 			channel,
@@ -998,7 +998,7 @@ export function bindBotNotifications(events: BotEvents, slackChannelOverrides: [
 	const userToken = (!args.devMode || args.useSlackInDev) && SLACK_TOKENS.user || args.devMode && args.useSlackInDev && SLACK_DEV_DUMMY_TOKEN
 	if (botToken && events.botConfig.slackChannel) {
 		logger.info('Enabling Slack messages for ' +  events.botname)
-		slackMessages = new SlackMessages(new Slack({id: events.botConfig.slackChannel, botToken, userToken}, args.slackDomain), persistence, logger)
+		slackMessages = new SlackMessages(new Slack({id: events.botConfig.slackChannel, botToken, userToken}, args.slackDomain, logger), persistence, logger)
 	}
 		
 	events.registerHandler(new BotNotifications(events.botname, events.botConfig.slackChannel, externalUrl, blockageUrlGenerator, logger, slackMessages, slackChannelOverrides))
diff --git a/Engine/Extras/RoboMerge/v3/src/robo/slack.ts b/Engine/Extras/RoboMerge/v3/src/robo/slack.ts
--- a/Engine/Extras/RoboMerge/v3/src/robo/slack.ts
+++ b/Engine/Extras/RoboMerge/v3/src/robo/slack.ts
@@ -1,6 +1,7 @@
 // Copyright Epic Games, Inc. All Rights Reserved.
 
 import * as request from '../common/request';
+import { ContextualLogger } from '../common/logger';
 
 export interface SlackChannel {
 	id: string
@@ -71,7 +72,7 @@ export interface SlackMessage {
 const MAIN_MESSAGE_FIELDS = new Set(['username', 'icon_emoji', 'channel', 'target', 'cl']);
 
 export class Slack {
-	constructor(private channel: SlackChannel, private domain: string) {
+	constructor(private channel: SlackChannel, private domain: string, private readonly logger: ContextualLogger) {
 	}
 
 	async addUserToChannel(user: string, channel: string, externalUser?: boolean) {
@@ -142,7 +143,8 @@ export class Slack {
 		catch {
 		}
 
-		throw new Error(`${command} generated:\n\t${resultJson}`)
+		this.logger.error(`${command} generated:\n\t${resultJson}`)
+		return {ok: false}
 	}
 
 	/*private*/ async post(command: string, args: any, canFail? : boolean) {
@@ -176,25 +178,8 @@ export class Slack {
 		catch {
 		}
 		
-		throw new Error(`url: '${url}' error: '${rawResult}'`)
-	}
-
-	async* getPages(command: string, limit?: number, inArgs?: any) {
-		const args = inArgs || {}
-		args.limit = limit || 100
-		for (let pageNum = 1;; ++pageNum) {
-			// for now limit to 50 pages, just in case
-			if (pageNum > 50) {throw new Error('busted safety valve!')}
-
-			const result = await this.get(command, args)
-			yield [result, pageNum]
-
-			args.cursor = result.response_metadata && result.response_metadata.next_cursor
-			if (!args.cursor) {
-				break
-			}
-		}
-
+		this.logger.error(`url: '${url}' error: '${rawResult}'`)
+		return {ok: false}
 	}
 
 // behaviour seems to be: put in attachment for green margin if any opts sent at all (e.g. {} different to nothing passed)
