Fix functional tests failing to wait for stomp operations to complete especially when running only a single test

diff --git a/Engine/Extras/RoboMerge/v3/functional_tests/framework/src/framework.ts b/Engine/Extras/RoboMerge/v3/functional_tests/framework/src/framework.ts
--- a/Engine/Extras/RoboMerge/v3/functional_tests/framework/src/framework.ts
+++ b/Engine/Extras/RoboMerge/v3/functional_tests/framework/src/framework.ts
@@ -430,25 +430,25 @@ export abstract class FunctionalTest {
 	}
 
 	async ensureNotBlocked(sourceStream: string, targetStream?: string) {
-		const branchState = await this.getBranchState(sourceStream)
-		if (targetStream) {
-			const edgeDisplayName = `${sourceStream} -> ${targetStream}`
-			this.info(`Ensuring ${edgeDisplayName} is not blocked`)
-			if (branchState.isBlocked()) {
-				throw new Error(`${sourceStream} (node) is blocked!`)
-			}
+			const branchState = await this.getBranchState(sourceStream)
+			if (targetStream) {
+				const edgeDisplayName = `${sourceStream} -> ${targetStream}`
+				this.info(`Ensuring ${edgeDisplayName} is not blocked`)
+				if (branchState.isBlocked()) {
+					throw new Error(`${sourceStream} (node) is blocked!`)
+				}
 
-			const edgeState = branchState.getEdgeState(this.fullBranchName(targetStream))
-			if (edgeState.isBlocked()) {
-				throw new Error(`${edgeDisplayName} is blocked!`)
+				const edgeState = branchState.getEdgeState(this.fullBranchName(targetStream))
+				if (edgeState.isBlocked()) {
+					throw new Error(`${edgeDisplayName} is blocked!`)
+				}
 			}
-		}
-		else {
-			this.info(`Ensuring ${sourceStream} isn't blocked`)
-			if (branchState.isBlocked()) {
-				throw new Error(`${sourceStream} (node) is blocked!`)
+			else {
+				this.info(`Ensuring ${sourceStream} isn't blocked`)
+				if (branchState.isBlocked()) {
+					throw new Error(`${sourceStream} (node) is blocked!`)
+				}
 			}
-		}
 	}
 
 	/**
@@ -1015,7 +1015,7 @@ export abstract class FunctionalTest {
 
 		// when idle, produce a map of edge names to tick counts for edges neither blocked nor waiting for a gate
 		const ticksToWaitFor = new Map<string, number>([...unblockedBranchStates]
-			.filter(([_, node]) => !node.getEdges().every(e => e.getGateClosedMessage() || e.isBlocked()))
+			.filter(([_, node]) => !node.getEdges().every(e => e.getGateClosedMessage()))
 			.map(([name, node]) => [name, node.getTickCount()]))
 
 		return ticksToWaitFor.size > 0 ? ticksToWaitFor : true
