Try and add additional diagnosis information for the swarm changelist issues with stomp

diff --git a/Engine/Extras/RoboMerge/v3/src/common/perforce.ts b/Engine/Extras/RoboMerge/v3/src/common/perforce.ts
--- a/Engine/Extras/RoboMerge/v3/src/common/perforce.ts
+++ b/Engine/Extras/RoboMerge/v3/src/common/perforce.ts
@@ -516,13 +516,14 @@ export class PerforceContext {
 
 	/** get a single change in the format of changes() */
 	async getChange(path_in: string, changenum: number, status?: ChangelistStatus) {
-		const list = await this.changes(`${path_in}@${changenum},${changenum}`, -1, 1, status) as Change[]
+		const list = await this.changes(`${path_in}@${changenum},${changenum}`, -1, 1, status, false) as Change[]
 		if (list.length <= 0) {
 			throw new Error(`Could not find changelist ${changenum} in ${path_in}`);
 		}
 		if (list.length > 1 || list[0].change !== changenum) {
 			// log for now
-			this.logger.error('p4.getChange unexpected result' +
+			const e = new Error();
+			this.logger.error(`${e.stack}\np4.getChange unexpected result for ${changenum}` +
 				list.map(change => `\n    ${change.change}: user ${change.user}, workspace ${change.client}`).join('')
 			)
 		}
@@ -533,14 +534,14 @@ export class PerforceContext {
 	 * Get a list of changes in a path since a specific CL
 	 * @return Promise to list of changelists
 	 */
-	changes(path_in: string, since: number, limit?: number, status?: ChangelistStatus): Promise<Change[]> {
+	changes(path_in: string, since: number, limit?: number, status?: ChangelistStatus, quiet?: boolean): Promise<Change[]> {
 		const path = since > 0 ? path_in + '@>' + since : path_in;
 		const args = ['changes', '-l',
 			(status ? `-s${status}` : '-ssubmitted'),
 			...(limit ? [`-m${limit}`] : []),
 			path];
 
-		return this.execAndParse(null, args, {quiet: true}, {
+		return this.execAndParse(null, args, {quiet: quiet ? quiet : true}, {
 			expected: {change: 'integer', client: 'string', user: 'string', desc: 'string'},
 			optional: {shelved: 'integer', oldChange: 'integer', IsPromoted: 'integer'}
 		}) as Promise<unknown> as Promise<Change[]>
