Functional test debugging

diff --git a/Engine/Extras/RoboMerge/v3/functional_tests/framework/src/framework.ts b/Engine/Extras/RoboMerge/v3/functional_tests/framework/src/framework.ts
--- a/Engine/Extras/RoboMerge/v3/functional_tests/framework/src/framework.ts
+++ b/Engine/Extras/RoboMerge/v3/functional_tests/framework/src/framework.ts
@@ -736,6 +736,7 @@ export abstract class FunctionalTest {
 		return getJson(DUMMY_SLACK_DOMAIN + '/posted/' + channel)
 			.then((perTargetThreads: [string, DummyThread[]][]) => {
 				for (const [target, dummyThreads] of perTargetThreads) {
+					console.log(`${targetName} : ${JSON.stringify(dummyThreads)}`)
 					if (target === targetName) {
 						return dummyThreads.filter(dt => dt.CL == cl).length == 1
 					}
