Additional functional test debugging

diff --git a/Engine/Extras/RoboMerge/v3/src/common/dummyslackserver.ts b/Engine/Extras/RoboMerge/v3/src/common/dummyslackserver.ts
--- a/Engine/Extras/RoboMerge/v3/src/common/dummyslackserver.ts
+++ b/Engine/Extras/RoboMerge/v3/src/common/dummyslackserver.ts
@@ -44,6 +44,8 @@ export class DummySlackApp implements AppInterface {
 			return { dataObj }
 		}
 
+		console.log(`dummy-slack POST: ${command}`)
+
 		if (command === "conversations.invite") {
 			return { ok: true }
 		}
@@ -69,10 +71,12 @@ export class DummySlackApp implements AppInterface {
 				return dataResult
 			}
 			const dataObj = dataResult.dataObj
+			console.log(`${JSON.stringify(dataObj)}`)
 
 			if (dataObj.thread_ts) {
 				let dummyThread = threads.get(dataObj.channel)?.get(dataObj.thread_ts)!
 				dummyThread.messages.push(dataObj.attachments[0].text)
+				console.log(`${JSON.stringify(posted)}`)
 				return { ok: true }
 			}
 			else {
@@ -94,6 +98,7 @@ export class DummySlackApp implements AppInterface {
 					setDefault(threads, dataObj.channel, new Map)
 						.set(dummyThread.timestamp, dummyThread)
 
+					console.log(`${JSON.stringify(posted)}`)
 					return { ok: true, ts: dummyThread.timestamp }
 				}
 				else {
@@ -110,6 +115,8 @@ export class DummySlackApp implements AppInterface {
 
 	@Handler('GET', '/posted/*')
 	channelMessages(channel: string) {
+		console.log(`/posted/${channel}`)
+		console.log(JSON.stringify(posted))
 		return [...(posted.get(channel) || [])]
 	}
 
