Remove the debugging

diff --git a/Engine/Extras/RoboMerge/v3/functional_tests/framework/src/framework.ts b/Engine/Extras/RoboMerge/v3/functional_tests/framework/src/framework.ts
--- a/Engine/Extras/RoboMerge/v3/functional_tests/framework/src/framework.ts
+++ b/Engine/Extras/RoboMerge/v3/functional_tests/framework/src/framework.ts
@@ -736,7 +736,6 @@ export abstract class FunctionalTest {
 		return getJson(DUMMY_SLACK_DOMAIN + '/posted/' + channel)
 			.then((perTargetThreads: [string, DummyThread[]][]) => {
 				for (const [target, dummyThreads] of perTargetThreads) {
-					this.info(`${targetName} : ${JSON.stringify(dummyThreads)}`)
 					if (target === targetName) {
 						return dummyThreads.filter(dt => dt.CL == cl).length == 1
 					}
diff --git a/Engine/Extras/RoboMerge/v3/src/common/dummyslackserver.ts b/Engine/Extras/RoboMerge/v3/src/common/dummyslackserver.ts
--- a/Engine/Extras/RoboMerge/v3/src/common/dummyslackserver.ts
+++ b/Engine/Extras/RoboMerge/v3/src/common/dummyslackserver.ts
@@ -44,8 +44,6 @@ export class DummySlackApp implements AppInterface {
 			return { dataObj }
 		}
 
-		console.log(`dummy-slack POST: ${command}`)
-
 		if (command === "conversations.invite") {
 			return { ok: true }
 		}
@@ -71,12 +69,10 @@ export class DummySlackApp implements AppInterface {
 				return dataResult
 			}
 			const dataObj = dataResult.dataObj
-			console.log(`${JSON.stringify(dataObj)}`)
 
 			if (dataObj.thread_ts) {
 				let dummyThread = threads.get(dataObj.channel)?.get(dataObj.thread_ts)!
 				dummyThread.messages.push(dataObj.attachments[0].text)
-				console.log(`${JSON.stringify(posted)}`)
 				return { ok: true }
 			}
 			else {
@@ -98,7 +94,6 @@ export class DummySlackApp implements AppInterface {
 					setDefault(threads, dataObj.channel, new Map)
 						.set(dummyThread.timestamp, dummyThread)
 
-					console.log(`${JSON.stringify(posted)}`)
 					return { ok: true, ts: dummyThread.timestamp }
 				}
 				else {
@@ -115,8 +110,6 @@ export class DummySlackApp implements AppInterface {
 
 	@Handler('GET', '/posted/*')
 	channelMessages(channel: string) {
-		console.log(`/posted/${channel}`)
-		console.log(JSON.stringify(posted))
 		return [...(posted.get(channel) || [])]
 	}
 
@@ -129,4 +122,4 @@ export class DummySlackApp implements AppInterface {
 		return { ok: false }
 	}
 
-}
+}
\ No newline at end of file
