Update documentation for yield

diff --git a/Engine/Extras/RoboMerge/v3/README.md b/Engine/Extras/RoboMerge/v3/README.md
--- a/Engine/Extras/RoboMerge/v3/README.md
+++ b/Engine/Extras/RoboMerge/v3/README.md
@@ -88,6 +88,7 @@ for both node and edge.
 |name                      |used by|default|edge?|description                           |notes|
 |--------------------------|:-----:|-------|-----|--------------------------------------|-----|
 |`checkIntervalSecs`        |b |`30`        | |All edges are round-robinned within this time, unless delayed by integrations| |
+|`yieldFrequency`           |b |`10`        | |How many integrations to complete before considering the next round-robin| |
 |`noStreamAliases`          |b |`false`     | |Stream names not available for commands, e.g. if there are duplicate stream names in bot | |
 |`reportToBuildHealth`      |b |`false`     | |UGS integration                          | |
 |`slackChannel`             |b |            | |Channel to receive blockages             |Not practical to make this per edge|
