We don't need the ignore

diff --git a/Engine/Extras/RoboMerge/v3/src/robo/autobranchupdater.ts b/Engine/Extras/RoboMerge/v3/src/robo/autobranchupdater.ts
--- a/Engine/Extras/RoboMerge/v3/src/robo/autobranchupdater.ts
+++ b/Engine/Extras/RoboMerge/v3/src/robo/autobranchupdater.ts
@@ -232,7 +232,7 @@ export class AutoBranchUpdater implements Bot {
 		const {depotpath, realFilepath, mirrorFilepath} = workspace
 
 		await this.p4.sync(workspace, depotpath, {opts:[P4_FORCE]})
-		const cl = await this.p4.new_cl(workspace, "Updating mirror file\n#jira none\n#robomerge ignore\n")
+		const cl = await this.p4.new_cl(workspace, "Updating mirror file\n#jira none\n")
 		await this.p4.edit(workspace, cl, depotpath)
 		await new Promise((done, _) => fs.copyFile(realFilepath, mirrorFilepath, done))
 		await this.p4.submit(workspace, cl)
