Fix crash reconsidering a CL with assets on an edge that blocks asset merges

diff --git a/Engine/Extras/RoboMerge/v3/src/robo/nodebot.ts b/Engine/Extras/RoboMerge/v3/src/robo/nodebot.ts
--- a/Engine/Extras/RoboMerge/v3/src/robo/nodebot.ts
+++ b/Engine/Extras/RoboMerge/v3/src/robo/nodebot.ts
@@ -1583,9 +1583,10 @@ export class NodeBot extends PerforceStatefulBot implements NodeBotInterface {
 		// this will deal with #manual changes with a single target and will also manage to work for multiple
 		// targets as long as all the target workspaces are on the same edge, but it will still crash if you
 		// have multiple targets with workspaces on different edgeservers
-		let isManualChange = result.info.targets ? result.info.targets![0].flags.has('manual') : false;
+		const numTargets = (result.info.targets || []).length;
+		let isManualChange = numTargets > 0 ? result.info.targets![0].flags.has('manual') : false;
 		if (change.forceCreateAShelf || (change.isUserRequest && !change.forceStompChanges) || isManualChange) {
-			if (!optWorkspaceOverride && (result.info.targets || []).length == 1) {
+			if (!optWorkspaceOverride && numTargets == 1) {
 				result.info.targetWorkspaceForShelf = await p4util.chooseBestWorkspaceForUser(this.p4, result.info.owner||result.info.author, result.info.targets![0].branch.stream)
 				optWorkspaceOverride = result.info.targetWorkspaceForShelf
 				this.nodeBotLogger.info(`Chose workspace ${optWorkspaceOverride}`)
