Robomerge unlock will now attempt to handle exclusive locks on moved files

diff --git a/Engine/Extras/RoboMerge/v3/src/common/perforce.ts b/Engine/Extras/RoboMerge/v3/src/common/perforce.ts
--- a/Engine/Extras/RoboMerge/v3/src/common/perforce.ts
+++ b/Engine/Extras/RoboMerge/v3/src/common/perforce.ts
@@ -22,6 +22,8 @@ const INTEGRATION_FAILURE_REGEXES: [RegExp, string][] = [
 
 export const EXCLUSIVE_CHECKOUT_REGEX = INTEGRATION_FAILURE_REGEXES[0][0]
 
+const REVERT_FAILURE_DUE_TO_MOVE_REGEX: RegExp = /(.*)#[0-9]+ - has been moved, not reverted/
+
 const changeResultExpectedShape: ztag.ParseOptions = {
 	expected: {change: 'integer', client: 'string', user: 'string', desc: 'string', time: 'integer', status: 'string', changeType: 'string'},
 	optional: {oldChange: 'integer'}
@@ -1185,7 +1187,22 @@ export class PerforceContext {
 		const edgeServer = await this.getWorkspaceEdgeServer(client)
 		const args = ['revert', '-C', client, ...files]
 		try {
-			await this._execP4Ztag(null, args, {edgeServerAddress: edgeServer?.address})
+			const results = await this._execP4Ztag(null, args, {edgeServerAddress: edgeServer?.address})
+
+			// Files can fail to be reverted because of moves, lets try and resolve that
+			const movedFiles = 
+				results.filter((result: any) => !result.action)
+					   .reduce((matches: string[], results: string[]) => 
+					   		matches.concat(
+								results.map(result => result.match(REVERT_FAILURE_DUE_TO_MOVE_REGEX))
+								       .filter(match => match)
+								       .map(match => match![1])
+							), [])
+			if (movedFiles.length > 0) {
+				const openedFiles = await this._execP4Ztag(null, ['opened','-a', ...movedFiles])
+				const pairedAdds: string[] = openedFiles.map((openedFile: any) => openedFile.movedFile)
+				await this.revertFiles(pairedAdds, client)
+			}
 		}
 		catch (reason) {
 			if (!isExecP4Error(reason)) {
