latestChange can be workspace based

diff --git a/Engine/Extras/RoboMerge/v3/src/common/perforce.ts b/Engine/Extras/RoboMerge/v3/src/common/perforce.ts
--- a/Engine/Extras/RoboMerge/v3/src/common/perforce.ts
+++ b/Engine/Extras/RoboMerge/v3/src/common/perforce.ts
@@ -256,7 +256,7 @@ export interface Workspace {
 }
 
 // temporary fudging of workspace string used by main Robo code
-export type RoboWorkspace = Workspace | string | null;
+export type RoboWorkspace = Workspace | string | null | undefined;
 
 export interface ClientSpec {
 	client: string
@@ -549,14 +549,14 @@ export class PerforceContext {
 		}) as Promise<unknown> as Promise<Change[]>
 	}
 
-	async latestChange(path: string): Promise<Change> {
+	async latestChange(path: string, workspace?: RoboWorkspace): Promise<Change> {
 
 		// temporarily waiting 30 seconds - filing ticket   - was: wait no longer than 5 seconds, retry up to 3 times
 		const args = ['-vnet.maxwait=30', '-r3', 'changes', '-l', '-ssubmitted', '-m1', path]
 
 		const startTime = Date.now()
 
-		const result = await this.execAndParse(null, args, {quiet: true, trace: true}, changeResultExpectedShape)
+		const result = await this.execAndParse(workspace, args, {quiet: true, trace: true}, changeResultExpectedShape)
 		if (!result || result.length !== 1) {
 			throw new Error("Expected exactly one change")
 		}
@@ -795,14 +795,14 @@ export class PerforceContext {
 			AltRoots: roots
 		}
 
-		// Perforce paths are mighty particular 
-		if (bsDepotPath.endsWith("/...")) {
+							// Perforce paths are mighty particular 
+							if (bsDepotPath.endsWith("/...")) {
 			params.View = bsDepotPath + ` //${workspace.name}/...`
-		} else if (bsDepotPath.endsWith("/")) {
+							} else if (bsDepotPath.endsWith("/")) {
 			params.View = bsDepotPath + `... //${workspace.name}/...`
-		} else {
+							} else {
 			params.View = bsDepotPath + `/... //${workspace.name}/...`
-		}
+							}
 
 		return this.newWorkspace(workspace.name, params);
 	}
@@ -1735,7 +1735,6 @@ export function getRootDirectoryForBranch(name: string): string {
 	return process.platform === "win32" ? `d:/ROBO/${name}` : `/src/${name}`;
 }
 
-// temp (ha)
 export function coercePerforceWorkspace(workspace: any): Workspace | null {
 	if (!workspace)
 		return null
