Robomerge can now handle streams of arbitrary StreamDepth for branchspecs and persistence

diff --git a/Engine/Extras/RoboMerge/v3/src/common/perforce.ts b/Engine/Extras/RoboMerge/v3/src/common/perforce.ts
--- a/Engine/Extras/RoboMerge/v3/src/common/perforce.ts
+++ b/Engine/Extras/RoboMerge/v3/src/common/perforce.ts
@@ -184,6 +184,9 @@ export function parseZTag(buffer: string, opts?: ExecZtagOpts) {
 			output.push(text);
 	}
 
+	if (opts && opts.reduce) {
+		return output.reduce((accumulator: any, value: any) => { accumulator = {...accumulator, ...value}; return accumulator; }, {})
+	}
 	return output;
 }
 
@@ -291,6 +294,7 @@ interface ExecOpts {
 interface ExecZtagOpts extends ExecOpts {
 	multiline?: boolean;
 	resolve?: boolean; // hacky solution to clear certain problematic lines out of resolve ztags
+	reduce?: boolean; // collapse the multiple entries to a single object, useful for problem parses that end up in multiple entries despite being a single result
 }
 
 export interface EditChangeOpts {
@@ -575,6 +579,36 @@ export class PerforceContext {
 		return this.execAndParse(null, args, {quiet: true}, changeResultExpectedShape) as Promise<unknown> as Promise<Change[]>
 	}
 
+	async getDepot(depotName: string) {
+		if ((await this._execP4Ztag(null, ['depots', '-e', depotName])).length > 0)
+		{
+			return this._execP4Ztag(null, ['depot', '-o', depotName], {reduce: true})
+		}
+		return null
+	}
+
+	async getStreamName(path: string) {
+		// Given the path, determine the depot and stream that a workspace needs to be created for
+		let depotEndChar = path.indexOf('/',2)
+		if (depotEndChar == -1) {
+			return new Error(`Unable to determine depot from $(path)`)
+		}
+
+		const depot = await this.getDepot(path.substring(2,depotEndChar))
+		if (depot.Type == 'stream') {
+			let streamDepth = depot.StreamDepth.match(/\//g).length - 2
+			if (streamDepth < 1) {
+				streamDepth = Number(depot.StreamDepth)
+			}
+			let streamNameEnd = 2;
+			for (let i=0;i<streamDepth+1;i++) {
+				streamNameEnd = path.indexOf("/",streamNameEnd+1)
+			}
+			return path.substring(0,streamNameEnd)
+		}
+		return new Error(`Depot $(depot) is not of type stream`)
+	}
+
 	async streams() {
 		const rawStreams = await this.execAndParse(null, ['streams'], {quiet: true}, {
 			expected: {Stream: 'string', Update: 'integer', Access: 'integer', Owner: 'string', Name: 'string', Parent: 'string', Type: 'string', desc: 'string',
@@ -795,14 +829,14 @@ export class PerforceContext {
 			AltRoots: roots
 		}
 
-							// Perforce paths are mighty particular 
-							if (bsDepotPath.endsWith("/...")) {
+		// Perforce paths are mighty particular 
+		if (bsDepotPath.endsWith("/...")) {
 			params.View = bsDepotPath + ` //${workspace.name}/...`
-							} else if (bsDepotPath.endsWith("/")) {
+		} else if (bsDepotPath.endsWith("/")) {
 			params.View = bsDepotPath + `... //${workspace.name}/...`
-							} else {
+		} else {
 			params.View = bsDepotPath + `/... //${workspace.name}/...`
-							}
+		}
 
 		return this.newWorkspace(workspace.name, params);
 	}
diff --git a/Engine/Extras/RoboMerge/v3/src/robo/robo.ts b/Engine/Extras/RoboMerge/v3/src/robo/robo.ts
--- a/Engine/Extras/RoboMerge/v3/src/robo/robo.ts
+++ b/Engine/Extras/RoboMerge/v3/src/robo/robo.ts
@@ -407,29 +407,26 @@ async function init(logger: ContextualLogger) {
 
 	let lookupStream = async function(rootPath: string): Promise<string> {
 
-		const match = rootPath.match(/(\/\/.*?\/.*?)(\/|$)/)
-		if (match) {
-			const stream = match[1]
-
-			while (true) {
-				try {
+		while (true) {
+			try {
+				const stream = await robo.p4.getStreamName(rootPath);
+				if (typeof stream === 'string') {
 					const streams = await robo.p4.streams();
 					if (streams.has(stream)) {
-						break
+						return stream
 					}
 				}
-				catch (err) {
+				else {
+					logger.warn(stream.message)
+					return ""
 				}
-
-				const timeout = 5.0;
-				logger.info(`Will check for ${stream} again in ${timeout} sec...`);
-				await _setTimeout(timeout*1000);
 			}
-			return stream
-		}
-		else {
-			logger.warn(`Unable to determine stream from root path '${rootPath}'`)
-			return ""
+			catch (err) {
+			}
+
+			const timeout = 5.0;
+			logger.info(`Will look up stream from ${rootPath} again in ${timeout} sec...`);
+			await _setTimeout(timeout*1000);
 		}
 	}
 
