Can now modify a manual gate from webpage
Manual gates now show "paused at gate" instead of "waiting for cis"

diff --git a/Engine/Extras/RoboMerge/v3/.vscode/launch.json b/Engine/Extras/RoboMerge/v3/.vscode/launch.json
--- a/Engine/Extras/RoboMerge/v3/.vscode/launch.json
+++ b/Engine/Extras/RoboMerge/v3/.vscode/launch.json
@@ -25,6 +25,7 @@
                 "ROBO_EXTERNAL_URL": "http://localhost:8877",
                 "NODE_ENV": "development",
                 "ROBO_BRANCHSPECS_DIRECTORY": "../Robomerge/data",
+                "ROBO_GATE_UPDATES_DIRECTORY": "../Robomerge/gateupdates",
                 "ROBO_VAULT_PATH": "./vault",
                 "EPIC_ENV": "dev",
                 "EPIC_DEPLOYMENT": "local",
@@ -61,6 +62,7 @@
                 "ROBO_EXTERNAL_URL": "http://localhost:8877",
                 "NODE_ENV": "development",
                 "ROBO_BRANCHSPECS_DIRECTORY": "./debugdata",
+                "ROBO_GATE_UPDATES_DIRECTORY": "./debuggateupdates",
                 "ROBO_VAULT_PATH": "./vault",
                 "EPIC_ENV": "dev",
                 "EPIC_DEPLOYMENT": "local",
@@ -88,6 +90,7 @@
                 "ROBO_EXTERNAL_URL": "https://localhost:4433",
                 "NODE_ENV": "development",
                 "ROBO_BRANCHSPECS_DIRECTORY": "./debugdata",
+                "ROBO_GATE_UPDATES_DIRECTORY": "./debuggateupdates",
                 "ROBO_VAULT_PATH": "./vault",
                 "EPIC_ENV": "dev",
                 "EPIC_DEPLOYMENT": "local",
diff --git a/Engine/Extras/RoboMerge/v3/public/js/boilerplate.js b/Engine/Extras/RoboMerge/v3/public/js/boilerplate.js
--- a/Engine/Extras/RoboMerge/v3/public/js/boilerplate.js
+++ b/Engine/Extras/RoboMerge/v3/public/js/boilerplate.js
@@ -909,7 +909,7 @@ function createEdgeRow(nodeData, edgeData, includeActions) {
 	}
 	columnArray.push(renderLastChangeCell_Common(nodeData.bot, nodeData.def.name, edgeData.last_cl, edgeAPIOp,
 														operationArgs, catchupText, edgeData.display_name))
-	postRenderLastChangeCell_Edge(columnArray[columnArray.length - 1], edgeData)
+	postRenderLastChangeCell_Edge(columnArray[columnArray.length - 1], edgeData, edgeAPIOp, operationArgs)
 
 	return columnArray
 }
@@ -1702,11 +1702,11 @@ function prettyDate(date) {
 	}
 }
 
-function postRenderLastChangeCell_Edge(lastChangeCell, edgeData) {
+function postRenderLastChangeCell_Edge(lastChangeCell, edgeData, operationFunction, operationArgs) {
 	if (edgeData.lastGoodCL) {
-		let tooltip = 'CL approved by CIS'
+		let tooltip = edgeData.lastGoodCLJobLink ? 'CL approved by CIS' : 'Paused at CL'
 		if (edgeData.lastGoodCLDate) {
-			tooltip += ` on ${prettyDate(new Date(edgeData.lastGoodCLDate))}`
+			tooltip += ` submitted ${prettyDate(new Date(edgeData.lastGoodCLDate))}`
 		}
 		if (edgeData.headCL) {
 			tooltip += ` (head changelist ${edgeData.headCL})`
@@ -1714,6 +1714,37 @@ function postRenderLastChangeCell_Edge(lastChangeCell, edgeData) {
 		
 		let goodCL = edgeData.lastGoodCLJobLink ? $(`<a href="${edgeData.lastGoodCLJobLink}">`).prop('target', '_blank') : $('<div>');
 		goodCL.html('\u{2713} ' + edgeData.lastGoodCL).addClass('last-good-cl').prop('title', tooltip).appendTo(lastChangeCell)
+
+		if (!edgeData.lastGoodCLJobLink)
+		{
+			// On shift+click, we can set the CL instead
+			goodCL.click(function(evt) {
+				if (evt.shiftKey)
+				{
+					let data = promptFor({
+						cl: {prompt: 'Enter CL', default: edgeData.lastGoodCL},
+					})
+					if (data) {
+						data.reason = "manually set through Robomerge homepage"
+
+						operationFunction(...operationArgs, "/set_gate_cl?" + toQuery(data), function(success) {
+							if (success) {
+								updateBranchList(botname)
+								displaySuccessfulMessage(`Successfully set gate for ${edgeData.displayName} to changelist ${data.cl}`)
+							} else {
+								displayErrorMessage(`Error setting gate for ${edgeData.displayName} to changelist ${data.cl}, please check logs.`)
+							}
+						})
+						
+					}
+					if (evt.preventDefault) {
+						evt.preventDefault()
+					}
+					return false
+				}
+				return true
+			})
+		}
 	}
 }
 
diff --git a/Engine/Extras/RoboMerge/v3/src/common/perforce.ts b/Engine/Extras/RoboMerge/v3/src/common/perforce.ts
--- a/Engine/Extras/RoboMerge/v3/src/common/perforce.ts
+++ b/Engine/Extras/RoboMerge/v3/src/common/perforce.ts
@@ -741,7 +741,7 @@ export class PerforceContext {
 		}
 		args.push(depotPath)
 		try {
-			await this._execP4(workspace, args)
+			return await this._execP4Ztag(workspace, args)
 		}
 		catch (reason) {
 			if (!isExecP4Error(reason)) {
@@ -755,6 +755,8 @@ export class PerforceContext {
 				}
 			}
 		}
+
+		return null
 	}
 
 	async syncAndReturnChangelistNumber(roboWorkspace: RoboWorkspace, depotPath: string, opts?: string[]) {
diff --git a/Engine/Extras/RoboMerge/v3/src/robo/bot-interfaces.ts b/Engine/Extras/RoboMerge/v3/src/robo/bot-interfaces.ts
--- a/Engine/Extras/RoboMerge/v3/src/robo/bot-interfaces.ts
+++ b/Engine/Extras/RoboMerge/v3/src/robo/bot-interfaces.ts
@@ -52,6 +52,7 @@ export interface IPCControls {
 	unacknowledge(changeCl: number) : OperationResult
 	
 	forceSetLastClWithContext(value: number, culprit: string, reason: string, unblock?: boolean): number
+	setGateCl(value: number, culprit: string, reason: string): Promise<number | null>
 }
 
 export interface EdgeBotInterface {
diff --git a/Engine/Extras/RoboMerge/v3/src/robo/edgebot.ts b/Engine/Extras/RoboMerge/v3/src/robo/edgebot.ts
--- a/Engine/Extras/RoboMerge/v3/src/robo/edgebot.ts
+++ b/Engine/Extras/RoboMerge/v3/src/robo/edgebot.ts
@@ -976,6 +976,10 @@ class EdgeBotImpl extends PerforceStatefulBot {
 		return prevValue
 	}
 
+	setGateCl(value: number, culprit: string, reason: string) {
+		return this.gate.setGateCl(value, culprit, reason)
+	}
+
 	reconsider(instigator: string, changeCl: number, additionalArgs?: Partial<ReconsiderArgs>) {
 		this.sourceNode.reconsider(instigator, changeCl, {targetBranchName: this.targetBranch.name, ...(additionalArgs || {})})
 	}
@@ -999,6 +1003,7 @@ abstract class EdgeBotEntryPoints implements IPCControls {
 	acknowledge: EdgeBotImpl["acknowledge"]
 	unacknowledge: EdgeBotImpl["unacknowledge"]
 	forceSetLastClWithContext: EdgeBotImpl["forceSetLastClWithContext"]
+	setGateCl: EdgeBotImpl["setGateCl"]
 	resetIntegrationTimestamp: EdgeBotImpl["resetIntegrationTimestamp"]
 
 	// async methods
@@ -1041,6 +1046,7 @@ export class EdgeBot
 		this.acknowledge = this.proxy("acknowledge")
 		this.unacknowledge = this.proxy("unacknowledge")
 		this.forceSetLastClWithContext = this.proxy("forceSetLastClWithContext")
+		this.setGateCl = this.proxy("setGateCl")
 		this.resetIntegrationTimestamp = this.proxy("resetIntegrationTimestamp")
 
 		this.revertPendingCLWithShelf = this.proxyAsync("revertPendingCLWithShelf")
@@ -1131,7 +1137,8 @@ export class EdgeBot
 			reconsider: this.reconsider,
 			acknowledge: this.acknowledge,
 			unacknowledge: this.unacknowledge,
-			forceSetLastClWithContext: this.forceSetLastClWithContext
+			forceSetLastClWithContext: this.forceSetLastClWithContext,
+			setGateCl: this.setGateCl
 		}
 	}
 }
diff --git a/Engine/Extras/RoboMerge/v3/src/robo/gate.ts b/Engine/Extras/RoboMerge/v3/src/robo/gate.ts
--- a/Engine/Extras/RoboMerge/v3/src/robo/gate.ts
+++ b/Engine/Extras/RoboMerge/v3/src/robo/gate.ts
@@ -1,7 +1,9 @@
 // Copyright Epic Games, Inc. All Rights Reserved.
 
+import * as fs from 'fs';
+import * as path from 'path';
 import { ContextualLogger } from "../common/logger";
-import { Change, PerforceContext } from "../common/perforce";
+import { P4_FORCE, Change, PerforceContext, Workspace } from "../common/perforce";
 import { gatesSame, GateEventContext, GateInfo } from "./branch-interfaces"
 import { DAYS_OF_THE_WEEK, CommonOptionFields, IntegrationWindowPane } from "./branchdefs"
 import { BotEventTriggers } from "./events"
@@ -93,6 +95,15 @@ async function getRequestedGateCl(context: GateContext, previousGateInfo?: GateI
 	return result
 }
 
+function readFileToString(filename: string) {
+	try {
+		return fs.readFileSync(filename, 'utf8');
+	}
+	catch (e) {
+		return null;
+	}
+}
+
 type EventTriggersAndStuff = GateEventContext & {
 	eventTriggers: BotEventTriggers
 }
@@ -110,6 +121,20 @@ export class Gate {
 		this.loadFromPersistence()
 	}
 
+	static gateUpdateWorkspacePrefix: string
+	static gateUpdateDirectory: string
+
+	static init(gateUpdateWorkspacePrefix: string, gateUpdatesDirectory: string, logger: ContextualLogger) {
+		this.gateUpdateWorkspacePrefix = gateUpdateWorkspacePrefix
+
+		let workspaceDir : string = path.resolve(gateUpdatesDirectory);
+		if (!fs.existsSync(workspaceDir)) {
+			logger.info(`Creating local branchspec directory ${workspaceDir}`)
+			fs.mkdirSync(workspaceDir);
+		}
+		this.gateUpdateDirectory = workspaceDir
+	}
+
 	bypass = false
 
 	private processGateChange(newGate: GateInfo) {
@@ -370,8 +395,7 @@ export class Gate {
 		//	- integration window is preventing catching up to the next gate
 
 		if (this.currentGateInfo) {
-			// should show last cl in queue on web page, null means caught up
-			return this.lastCl >= this.currentGateInfo.cl ? 'waiting for CIS' : null
+			return this.lastCl >= this.currentGateInfo.cl ? this.currentGateInfo.link ? 'waiting for CIS' : 'paused at gate' : null
 		}
 
 		// null means no gate
@@ -409,6 +433,76 @@ export class Gate {
 		}
 	}
 
+	async setGateCl(value: number, culprit: string, reason: string) {
+
+		const lastGoodCLPath = this.context.options.lastGoodCLPath
+		if (typeof lastGoodCLPath !== "string" ) {
+			this.context.logger.error('Cannot set a Gate CL if lastGoodCLPath is not a perforce path')
+			return null
+		}
+	
+		const prevCl = await getRequestedGateCl(this.context)
+		if (prevCl == null) {
+			this.context.logger.error('Cannot set a Gate CL if lastGoodCLPath is not a perforce path')
+			return null
+		}
+	
+		if (value == prevCl.cl) {
+			this.context.logger.info('Not updating gate file as requested CL ${value} matches current value')
+			return prevCl.cl
+		}
+
+		const stream = await this.context.p4!.getStreamName(lastGoodCLPath);
+		if (typeof stream === 'string') {
+			const streams = await this.context.p4!.streams();
+			if (!streams.has(stream)) {
+				this.context.logger.error(`Unable to find stream '${stream}'`)
+				return null
+			}
+		}
+		else {
+			this.context.logger.error(stream.message)
+			return null
+		}
+	
+		const workspace: Workspace = {
+			name: `${Gate.gateUpdateWorkspacePrefix}-${stream.substring(2).replace("/","_")}`,
+			directory: Gate.gateUpdateDirectory
+		}
+	
+		const p4 = this.context.p4!
+	
+		const workspaceQueryResult = await p4.find_workspace_by_name(workspace.name)
+		if (workspaceQueryResult.length === 0) {
+			await p4.newWorkspace(workspace.name, {
+				Stream: stream,
+				Root: workspace.directory
+			})
+		}
+	
+		const syncResult = await p4.sync(workspace, lastGoodCLPath, {opts:[P4_FORCE]})
+		if (syncResult.length != 1) {
+			this.context.logger.error(`Unable to sync ${lastGoodCLPath} in workspace ${workspace.name}`)
+			return null
+		}
+		const realFilepath = syncResult[0]['clientFile']
+	
+		let fileContents = readFileToString(realFilepath)
+		if (!fileContents) {
+			this.context.logger.error(`Unable to read file ${realFilepath}`)
+			return null
+		}
+	
+		fileContents = fileContents.replace(`${prevCl.cl}`,`${value}`)
+	
+		const cl = await p4.new_cl(workspace, `${culprit} updating gate file to ${value}. Reason: ${reason}`)
+		await p4.edit(workspace, cl, lastGoodCLPath)
+		fs.writeFileSync(realFilepath, fileContents, "utf8")
+		await p4.submit(workspace, cl)
+	
+		return prevCl.cl
+	}
+
 	logSummary() {
 		const logger = this.context.logger
 		logger.info('current gate: ' + (this.currentGateInfo ? this.currentGateInfo.cl.toString() : 'none'))
diff --git a/Engine/Extras/RoboMerge/v3/src/robo/ipc.ts b/Engine/Extras/RoboMerge/v3/src/robo/ipc.ts
--- a/Engine/Extras/RoboMerge/v3/src/robo/ipc.ts
+++ b/Engine/Extras/RoboMerge/v3/src/robo/ipc.ts
@@ -592,6 +592,23 @@ export class IPC {
 														`requested by ${query.who} (Reason: ${query.reason})`)
 			return OPERATION_SUCCESS
 
+		case 'set_gate_cl':
+			if (!edgeName) {
+				return {statusCode: 400, message: 'Only valid to call for an edge'}
+			}
+
+			cl = parseInt(query.cl)
+
+			if (isNaN(cl)) {
+				return {statusCode: 400, message: 'Invalid CL parameter: ' + cl}
+			}
+
+			let prevGateCl = await generalOpTarget.setGateCl(cl, query.who, query.reason)
+
+			this.ipcLogger.info(`Setting gate CL=${cl} on ${botname} : ${branch.name} (was CL ${prevGateCl}), ` +
+														`requested by ${query.who} (Reason: ${query.reason})`)
+			return OPERATION_SUCCESS
+
 		case 'reconsider':
 			const argMatch = query.cl.match(/(\d+)\s*(.*)/)
 			if (argMatch) {
diff --git a/Engine/Extras/RoboMerge/v3/src/robo/nodebot.ts b/Engine/Extras/RoboMerge/v3/src/robo/nodebot.ts
--- a/Engine/Extras/RoboMerge/v3/src/robo/nodebot.ts
+++ b/Engine/Extras/RoboMerge/v3/src/robo/nodebot.ts
@@ -517,6 +517,12 @@ export class NodeBot extends PerforceStatefulBot implements NodeBotInterface {
 		return prevValue
 	}
 	
+	setGateCl(_1: number, _2: string, _3:string): Promise<number | null> {
+		const err = new Error('setGateCl not implemented on NodeBot')
+		this.nodeBotLogger.printException(err)
+		throw err
+	}
+
 	onForcedLastCl(nodeOrEdgeName: string, forcedCl: number, previousCl: number, culprit: string, reason: string) {
 		this.conflicts.onForcedLastCl({nodeOrEdgeName, forcedCl, previousCl, culprit, reason})
 	}
diff --git a/Engine/Extras/RoboMerge/v3/src/robo/perforce-stateful-bot.ts b/Engine/Extras/RoboMerge/v3/src/robo/perforce-stateful-bot.ts
--- a/Engine/Extras/RoboMerge/v3/src/robo/perforce-stateful-bot.ts
+++ b/Engine/Extras/RoboMerge/v3/src/robo/perforce-stateful-bot.ts
@@ -152,6 +152,7 @@ export abstract class PerforceStatefulBot implements BotIPC {
 	}
 
 	abstract forceSetLastClWithContext(value: number, culprit: string, reason: string, unblock: boolean): number;
+	abstract setGateCl(value: number, culprit: string, reason: string): Promise<number | null>;
 
 	async _getChange(changeCl: number, path?: string, status?: ChangelistStatus) : Promise<Change | string> {
 		try {
diff --git a/Engine/Extras/RoboMerge/v3/src/robo/robo.ts b/Engine/Extras/RoboMerge/v3/src/robo/robo.ts
--- a/Engine/Extras/RoboMerge/v3/src/robo/robo.ts
+++ b/Engine/Extras/RoboMerge/v3/src/robo/robo.ts
@@ -17,6 +17,7 @@ import { CertFiles } from '../common/webserver';
 import { addBranchGraph, Graph, GraphAPI } from '../new/graph';
 import { AutoBranchUpdater } from './autobranchupdater';
 import { Branch } from './branch-interfaces';
+import { Gate } from './gate';
 import { GraphBot } from './graphbot';
 import { IPC, Message } from './ipc';
 import { NodeBot } from './nodebot';
@@ -87,6 +88,18 @@ const COMMAND_LINE_ARGS: {[param: string]: Arg<any>} = {
 		dflt: './data'
 	},
 
+	gateUpdateWorkspacePrefix: {
+		match: /^-gate_update_workspace_prefix=(.+)$/,
+		env: 'ROBO_GATE_UPDATE_WORKSPACE_PREFIX',
+		dflt: 'robomerge-gateupdate-' + os.hostname()
+	},
+
+	gateUpdatesDirectory: {
+		match: /^-gate_updates_directory=(.+)$/,
+		env: 'ROBO_GATE_UPDATES_DIRECTORY',
+		dflt: './gateupdates'
+	},
+
 	noIPC: {
 		match: /^(-noIPC)$/,
 		parse: _str => true,
@@ -456,6 +469,12 @@ async function init(logger: ContextualLogger) {
 		logger.warn('Auto brancher updater not configured!')
 	}
 
+	if (args.gateUpdateWorkspacePrefix && args.gateUpdatesDirectory) {
+		Gate.init(args.gateUpdateWorkspacePrefix, args.gateUpdatesDirectory, logger)
+	} else {
+		logger.warn('Gate Update Workspace Prefix or Directory not specified. Gate updating via webpage may not work correctly.')
+	}
+
 	if (args.persistenceBackupFrequency > 0 && !args.previewOnly) {
 		const workspace: Object[] = await robo.p4.find_workspace_by_name(args.persistenceBackupWorkspace)
 		if (workspace.length === 0) {
