Fix missing definition

diff --git a/Engine/Extras/RoboMerge/v3/public/js/boilerplate.js b/Engine/Extras/RoboMerge/v3/public/js/boilerplate.js
--- a/Engine/Extras/RoboMerge/v3/public/js/boilerplate.js
+++ b/Engine/Extras/RoboMerge/v3/public/js/boilerplate.js
@@ -909,7 +909,7 @@ function createEdgeRow(nodeData, edgeData, includeActions) {
 	}
 	columnArray.push(renderLastChangeCell_Common(nodeData.bot, nodeData.def.name, edgeData.last_cl, edgeAPIOp,
 														operationArgs, catchupText, edgeData.display_name))
-	postRenderLastChangeCell_Edge(columnArray[columnArray.length - 1], edgeData, edgeAPIOp, operationArgs)
+	postRenderLastChangeCell_Edge(nodeData.bot, columnArray[columnArray.length - 1], edgeData, edgeAPIOp, operationArgs)
 
 	return columnArray
 }
@@ -1702,7 +1702,7 @@ function prettyDate(date) {
 	}
 }
 
-function postRenderLastChangeCell_Edge(lastChangeCell, edgeData, operationFunction, operationArgs) {
+function postRenderLastChangeCell_Edge(botname, lastChangeCell, edgeData, operationFunction, operationArgs) {
 	if (edgeData.lastGoodCL) {
 		let tooltip = edgeData.lastGoodCLJobLink ? 'CL approved by CIS' : 'Paused at CL'
 		if (edgeData.lastGoodCLDate) {
