Move some gate status querying inside gate applyStatus function

diff --git a/Engine/Extras/RoboMerge/v3/src/robo/edgebot.ts b/Engine/Extras/RoboMerge/v3/src/robo/edgebot.ts
--- a/Engine/Extras/RoboMerge/v3/src/robo/edgebot.ts
+++ b/Engine/Extras/RoboMerge/v3/src/robo/edgebot.ts
@@ -929,9 +929,6 @@ class EdgeBotImpl extends PerforceStatefulBot {
 
 		status.last_cl = this.lastCl
 		
-		// if (this.isForceFlow)
-		status.num_changes_remaining = this.gate.numChangesRemaining
-
 		status.is_active = this.isActive
 		status.is_available = this.isAvailable
 		status.is_blocked = this.isBlocked
diff --git a/Engine/Extras/RoboMerge/v3/src/robo/gate.ts b/Engine/Extras/RoboMerge/v3/src/robo/gate.ts
--- a/Engine/Extras/RoboMerge/v3/src/robo/gate.ts
+++ b/Engine/Extras/RoboMerge/v3/src/robo/gate.ts
@@ -406,6 +406,8 @@ export class Gate {
 		// for now, just equivalent of what was there before
 		// @todo info about intermediate gate
 
+		outStatus.num_changes_remaining = this.numChangesRemaining
+
 		const mostRecentGate = this.getMostRecentGate()
 		if (mostRecentGate) {
 			outStatus.lastGoodCL = mostRecentGate.cl
