Add link to slack thread(s) to the conflict info block on the webpage

diff --git a/Engine/Extras/RoboMerge/v3/public/js/boilerplate.js b/Engine/Extras/RoboMerge/v3/public/js/boilerplate.js
--- a/Engine/Extras/RoboMerge/v3/public/js/boilerplate.js
+++ b/Engine/Extras/RoboMerge/v3/public/js/boilerplate.js
@@ -1027,10 +1027,14 @@ function createPauseDivs(data, conflict) {
 			data.blockage && data.blockage.change ? data.blockage.change :
 			'unknown'
 		
-		const info = conflict && conflict.kind ? `<span class="pause-div-label">Cause:</span> <strong>${conflict.kind.toLowerCase()}</strong>` :
+		let info = conflict && conflict.kind ? `<span class="pause-div-label">Cause:</span> <strong>${conflict.kind.toLowerCase()}</strong>` :
 			data.blockage ? `<span class="pause-div-label">Blocked.</span> Type: ${data.blockage.type}<br /> Message: ${data.blockage.message}` :
 			`No info can be provided. Please contact Robomerge help.`
 
+		if (conflict.slackLinks) {
+			info += `<br>${conflict.slackLinks.map(link => `<a href="${link}" target="_blank">Slack Thread</a>`).join("<br>")}`
+		}
+
 		divs.push($('<div class="info-block conflict">')
 			.append(
 				$('<p>').html('<h6>Blocked on change ' + makeClLink(changelist) + `</h6> ${info}`)
@@ -1212,9 +1216,9 @@ function renderStatusCell_Common(statusCell, data) {
 			let acknowledgedSince = new Date(data.blockage.acknowledgedAt);
 			let [ackDurationStr, ackDurationColor] = printDurationInfo("Acknowledged", Date.now() - acknowledgedSince.getTime())
 			$('<div class="blockage-details">')
-				.css('color', ackDurationColor)
-				.html(`${ackDurationStr} by <strong>${data.blockage.acknowledger}</strong>`)
-				.insertAfter(blockageinfoDiv)
+									.css('color', ackDurationColor)
+									.html(`${ackDurationStr} by <strong>${data.blockage.acknowledger}</strong>`)
+									.insertAfter(blockageinfoDiv)
 		}
 		// Determine who is responsible for resolving this.
 		else {
diff --git a/Engine/Extras/RoboMerge/v3/src/robo/conflicts.ts b/Engine/Extras/RoboMerge/v3/src/robo/conflicts.ts
--- a/Engine/Extras/RoboMerge/v3/src/robo/conflicts.ts
+++ b/Engine/Extras/RoboMerge/v3/src/robo/conflicts.ts
@@ -9,6 +9,7 @@ import { EdgeBot } from './edgebot';
 import { BotEventHandler, BotEventTriggers } from './events';
 import { Context } from './settings';
 import { UGS } from './ugs';
+import { SlackMessages } from './notifications';
 
 function conflictMatches(conflict: PersistentConflict, sourceCl: number, branchArg?: BranchArg | null) {
 	const targetBranchName = branchArg && resolveBranchArg(branchArg, true)
@@ -407,7 +408,7 @@ export class Conflicts {
 		return null
 	}
 
-	applyStatus(conflictsForStatus: Partial<ConflictStatusFields>[]) {
+	applyStatus(conflictsForStatus: Partial<ConflictStatusFields>[], slackMessages?: SlackMessages) {
 		for (const conflict of this.conflicts) {
 			const conflictObj: Partial<ConflictStatusFields> = {}
 
@@ -426,6 +427,18 @@ export class Conflicts {
 			conflictObj.author = conflict.author
 			conflictObj.owner = conflict.owner
 
+			if (slackMessages) {
+				const findResult = slackMessages.findAll(conflict.cl, conflict.targetBranchName || conflict.kind)
+				for (const result of findResult.messages) {
+					if (result.messageOpts.channel.startsWith('C') && result.permalink) {
+						if (!conflictObj.slackLinks) {
+							conflictObj.slackLinks = []
+						}
+						conflictObj.slackLinks.push(`${result.permalink}?thread_ts=${result.timestamp}`)
+					}
+				}
+			}
+
 			conflictsForStatus.push(conflictObj)
 		}
 	}
diff --git a/Engine/Extras/RoboMerge/v3/src/robo/nodebot.ts b/Engine/Extras/RoboMerge/v3/src/robo/nodebot.ts
--- a/Engine/Extras/RoboMerge/v3/src/robo/nodebot.ts
+++ b/Engine/Extras/RoboMerge/v3/src/robo/nodebot.ts
@@ -1261,7 +1261,7 @@ export class NodeBot extends PerforceStatefulBot implements NodeBotInterface {
 		}
 
 		status.conflicts = [] as any[]
-		this.conflicts.applyStatus(status.conflicts)
+		this.conflicts.applyStatus(status.conflicts, this.slackMessages)
 
 		status.tick_count = this.tickCount
 
diff --git a/Engine/Extras/RoboMerge/v3/src/robo/notifications.ts b/Engine/Extras/RoboMerge/v3/src/robo/notifications.ts
--- a/Engine/Extras/RoboMerge/v3/src/robo/notifications.ts
+++ b/Engine/Extras/RoboMerge/v3/src/robo/notifications.ts
@@ -206,6 +206,7 @@ function formatResolution(info: PersistentConflict) {
 
 type SlackConflictMessage = {
 	timestamp: string
+	permalink: string
 	messageOpts: SlackMessage
 }
 
@@ -254,8 +255,10 @@ export class SlackMessages {
 			message.cl = cl
 
 			let timestamp
+			let permalink
 			try {
 				timestamp = await this.slack.postMessage(message)
+				permalink = await this.slack.getPermalink(timestamp, message.channel)
 			}
 			catch (err) {
 				this.smLogger.printException(err, 'Error talking to Slack')
@@ -264,7 +267,7 @@ export class SlackMessages {
 
 			// Used for messages we don't care to keep, currently the /api/test/directmessage endpoint
 			if (persistMessage) {
-				findResult.persistedMessages[findResult.conflictKey] = {timestamp, messageOpts: message}
+				findResult.persistedMessages[findResult.conflictKey] = {timestamp, permalink, messageOpts: message}
 				this.persistence.set(SLACK_MESSAGES_PERSISTENCE_KEY, findResult.persistedMessages)
 			}
 		}
diff --git a/Engine/Extras/RoboMerge/v3/src/robo/slack.ts b/Engine/Extras/RoboMerge/v3/src/robo/slack.ts
--- a/Engine/Extras/RoboMerge/v3/src/robo/slack.ts
+++ b/Engine/Extras/RoboMerge/v3/src/robo/slack.ts
@@ -86,6 +86,10 @@ export class Slack {
 		return this.get('conversations.info', {channel})
 	}
 
+	async getPermalink(thread_ts: string, channel: string) {
+		return (await this.get('chat.getPermalink', {channel, message_ts: thread_ts})).permalink
+	}
+
 	async postMessage(message: SlackMessage) {
 		return (await this.post('chat.postMessage', this.makeArgs(message))).ts
 	}
diff --git a/Engine/Extras/RoboMerge/v3/src/robo/status-types.ts b/Engine/Extras/RoboMerge/v3/src/robo/status-types.ts
--- a/Engine/Extras/RoboMerge/v3/src/robo/status-types.ts
+++ b/Engine/Extras/RoboMerge/v3/src/robo/status-types.ts
@@ -170,6 +170,7 @@ export type ConflictStatusFields = {
 	kind: FailureKind
 	author: string
 	owner: string
+	slackLinks: string[]
 }
 
 type GraphBotError = {
