Reduce the amount of opened commands that will be issued to identify the owners of exclusive locked files to avoid spamming the server
Remove some client side logging from the unlock page

diff --git a/Engine/Extras/RoboMerge/v3/public/js/unlock.js b/Engine/Extras/RoboMerge/v3/public/js/unlock.js
--- a/Engine/Extras/RoboMerge/v3/public/js/unlock.js
+++ b/Engine/Extras/RoboMerge/v3/public/js/unlock.js
@@ -92,7 +92,6 @@ function unlockVerify() {
             verifyOperation.done(function(success) {
                 // Unlock returns a JSON payload on success
                 const unlockJson = JSON.parse(success)
-                console.log(`Unlock Verification message: ${unlockJson.message}`)
 
                 // Visualize the json
                 visualizeUnlockVerification(requestedBranchCl, unlockJson)
@@ -200,7 +199,9 @@ function transitionDisplayUnlockResults() {
 // Create visualization of the file, with links to swarm
 function visualizeAuthor(author, files) {
     const authorDiv = $(`<div id="${author}" class="unlock-visual">`)
-    if (robomergeUser && author.toLowerCase() === robomergeUser.userName.toLowerCase()) {
+    if (author.length == 0) {
+        authorDiv.append($('<h3>').html(`Unable to determine locked by user (likely due to too many locked files)`))
+    } else if (robomergeUser && author.toLowerCase() === robomergeUser.userName.toLowerCase()) {
         authorDiv.append($('<h3>').html("Locked by <strong>you</strong>"))
     } else {
         authorDiv.append($('<h3>').html(`Locked by ${author}`))
@@ -219,7 +220,6 @@ function visualizeAuthor(author, files) {
 function visualizeUnlockVerification(requestedBranchCl, unlockJson) {
     // Debug
     const prettyJson = JSON.stringify(unlockJson, null, 2)
-    console.log(`Unlock Verification Debug JSON:\n${prettyJson}`)
     let linkToSwarm = makeClLink(requestedBranchCl, `CL ${requestedBranchCl}`)
 
     let authorDict = {}
diff --git a/Engine/Extras/RoboMerge/v3/src/common/perforce.ts b/Engine/Extras/RoboMerge/v3/src/common/perforce.ts
--- a/Engine/Extras/RoboMerge/v3/src/common/perforce.ts
+++ b/Engine/Extras/RoboMerge/v3/src/common/perforce.ts
@@ -1359,7 +1359,7 @@ export class PerforceContext {
 	// get the email (according to P4) for a specific user
 	async getEmail(username: string) {
 		let m = null
-		if (!username.startsWith('@')) {
+		if (username.length > 0 && !username.startsWith('@')) {
 			const output = await this._execP4(null, ['user', '-o', username]);
 			// look for the email field
 			m = output.match(/\nEmail:\s+([^\n]+)\n/);
diff --git a/Engine/Extras/RoboMerge/v3/src/robo/edgebot.ts b/Engine/Extras/RoboMerge/v3/src/robo/edgebot.ts
--- a/Engine/Extras/RoboMerge/v3/src/robo/edgebot.ts
+++ b/Engine/Extras/RoboMerge/v3/src/robo/edgebot.ts
@@ -24,6 +24,7 @@ import { getIntegrationOwner } from "./targets";
 
 const FAILED_CHANGELIST_PAUSE_TIMEOUT_SECONDS = 15 * 60
 const MAX_INTEGRATION_ERRORS_TO_ANALYZE = 5
+const MAX_OPENED_COMMANDS_TO_ISSUE = 100
 const DEPOT_FILE_REGEX = /^(.*[\\\/])(.*)/
 const JIRA_REGEX = /^\s*#jira\s+(.*)/i
 
@@ -261,26 +262,29 @@ class EdgeBotImpl extends PerforceStatefulBot {
 
 	private async analyzeIntegrationError(errors: string[]) {
 		
-		const openedRequests: [RegExpMatchArray, Promise<OpenedFileRecord[]>, Promise<OpenedFileRecord[]>][] = []
+		const results: ExclusiveFile[] = []
+		const openedRequests: [RegExpMatchArray, Promise<OpenedFileRecord[]>, boolean][] = []
 		for (const err of errors) {
 			const match = err.match(EXCLUSIVE_CHECKOUT_REGEX)
 			if (match) {
-				openedRequests.push([match, this.p4.opened(null, match[1] + match[2], true), this.p4.opened(null, match[1] + match[2])])
+				if (openedRequests.length < MAX_OPENED_COMMANDS_TO_ISSUE) {
+					openedRequests.push([match, this.p4.opened(null, match[1] + match[2], true), true])
+				} else {
+					results.push({depotPath: match[1] + match[2], name: match[2], user: "", client: ""})
+				}
 			}
 		}
 
-		const results: ExclusiveFile[] = []
-		for (const [match, exclusiveReq, addReq] of openedRequests) {
-			const recs = await exclusiveReq
+		for (const [match, request, exclusive] of openedRequests) {
+			const recs = await request
 			if (recs.length > 0) {
 				// should only be one, since we're looking for exclusive check-out errors
 				results.push({depotPath: match[1] + match[2], name: match[2], user: recs[0].user, client: recs[0].client})
-			} else {
-				const recs = await addReq
-				if (recs.length > 0) {
-					// should only be one, since we're looking for exclusive check-out errors
-					results.push({depotPath: match[1] + match[2], name: match[2], user: recs[0].user, client: recs[0].client})
-				}
+			} else if (exclusive && openedRequests.length < MAX_OPENED_COMMANDS_TO_ISSUE) {
+				// if we failed to find it as an exclusive check-out, try as non-exclusive which will find adds
+				openedRequests.push([match, this.p4.opened(null, match[1] + match[2]), false])
+			} else  {
+				results.push({depotPath: match[1] + match[2], name: match[2], user: "", client: ""})
 			}
 		}
 
diff --git a/Engine/Extras/RoboMerge/v3/src/robo/notifications.ts b/Engine/Extras/RoboMerge/v3/src/robo/notifications.ts
--- a/Engine/Extras/RoboMerge/v3/src/robo/notifications.ts
+++ b/Engine/Extras/RoboMerge/v3/src/robo/notifications.ts
@@ -551,22 +551,30 @@ export class BotNotifications implements BotEventHandler {
 				const exclusiveLockUsers = (blockage.failure.additionalInfo as ExclusiveLockInfo).exclusiveLockUsers
 				let text = ''
 
+				let unknownLockUsers = false
 				for (const exclusiveLockUser of exclusiveLockUsers) {
-					let exclusiveLockSlackUser
-					const exclusiveLockUserEmail = await exclusiveLockUser.userEmail
-					if (exclusiveLockUserEmail) {
-						exclusiveLockSlackUser = await this.slackMessages.getSlackUser(exclusiveLockUserEmail)
-						if (exclusiveLockSlackUser) {
-							exclusiveLockSlackUser = `<@${exclusiveLockSlackUser}> `
+					if (exclusiveLockUser.user.length > 0) {
+						let exclusiveLockSlackUser
+						const exclusiveLockUserEmail = await exclusiveLockUser.userEmail
+						if (exclusiveLockUserEmail) {
+							exclusiveLockSlackUser = await this.slackMessages.getSlackUser(exclusiveLockUserEmail)
+							if (exclusiveLockSlackUser) {
+								exclusiveLockSlackUser = `<@${exclusiveLockSlackUser}> `
+							}
+							usersToInvite.add(exclusiveLockUserEmail)
 						}
-						usersToInvite.add(exclusiveLockUserEmail)
-					}
-					if (!exclusiveLockSlackUser) {
-						exclusiveLockSlackUser = `@${exclusiveLockUser.user} `
+						if (!exclusiveLockSlackUser) {
+							exclusiveLockSlackUser = `@${exclusiveLockUser.user} `
+						}
+						text += exclusiveLockSlackUser
+					} else {
+						unknownLockUsers = true
 					}
-					text += exclusiveLockSlackUser
 				}
 				text += `\n\nPlease unlock the files blocking robomerge or work with ${blockage.owner} to resolve the conflict`
+				if (unknownLockUsers) {
+					text += "\n\nNote that some locked files did not have their owner determined and as such those owners may not have been tagged"
+				}
 				
 				const exclusiveCheckoutMessage: SlackMessage = {
 					text,
