No longer spam users with locked files each time the merge is retried

diff --git a/Engine/Extras/RoboMerge/v3/src/robo/notifications.ts b/Engine/Extras/RoboMerge/v3/src/robo/notifications.ts
--- a/Engine/Extras/RoboMerge/v3/src/robo/notifications.ts
+++ b/Engine/Extras/RoboMerge/v3/src/robo/notifications.ts
@@ -473,7 +473,7 @@ export class BotNotifications implements BotEventHandler {
 	/** Conflict */
 	// considered if no Slack set up, should continue to add people to notify, but too complicated:
 	//		let's go all in on Slack. Fine for this to be async (but fire and forget)
-	async onBlockage(blockage: Blockage) {
+	async onBlockage(blockage: Blockage, isNew: boolean) {
 		const changeInfo = blockage.change
 
 		if (changeInfo.userRequest) {
@@ -547,7 +547,7 @@ export class BotNotifications implements BotEventHandler {
 
 			messagesToPost.push({ message })
 
-			if (blockage.failure.kind === 'Exclusive check-out') {
+			if (isNew && blockage.failure.kind === 'Exclusive check-out') {
 				const exclusiveLockUsers = (blockage.failure.additionalInfo as ExclusiveLockInfo).exclusiveLockUsers
 				let text = ''
 
