Fix interface property name

diff --git a/Engine/Extras/RoboMerge/v3/src/robo/branch-interfaces.ts b/Engine/Extras/RoboMerge/v3/src/robo/branch-interfaces.ts
--- a/Engine/Extras/RoboMerge/v3/src/robo/branch-interfaces.ts
+++ b/Engine/Extras/RoboMerge/v3/src/robo/branch-interfaces.ts
@@ -70,7 +70,7 @@ export interface ExclusiveLockUser {
 
 export interface ExclusiveLockInfo {
 	exclusiveLockUsers: ExclusiveLockUser[]
-	lockedFiles: ExclusiveFile[]
+	exclusiveFiles: ExclusiveFile[]
 }
 
 export interface UnlockVerification extends OperationResult {
