Remove obsolete code for a no longer existing "too many files"/maxFilesPerIntegration feature

diff --git a/Engine/Extras/RoboMerge/v3/README.md b/Engine/Extras/RoboMerge/v3/README.md
--- a/Engine/Extras/RoboMerge/v3/README.md
+++ b/Engine/Extras/RoboMerge/v3/README.md
@@ -113,7 +113,6 @@ for both node and edge.
 |`isDefaultBot`             |e |`false`     | |Run plain #robomerge commands? Should be `false` for streams monitored by multiple bots| |
 |`lastGoodCLPath`           |e |            | |'Gate' file to read to find CIS-approved CL| |
 |`waitingForCISLink`        |e |            | |URL to link to for waiting for CIS messages| |
-|`maxFilesPerIntegration`   |e |`-1`        | |Reject integrations with more files than this| |
 |`notify`                   |e |            | |Additional people to email on blockages  |Also `globalNotify`|
 |`resolver`                 |e |            | |Single designated resolver               |Currently applies to both source and target nodes|
 
diff --git a/Engine/Extras/RoboMerge/v3/public/js/boilerplate.js b/Engine/Extras/RoboMerge/v3/public/js/boilerplate.js
--- a/Engine/Extras/RoboMerge/v3/public/js/boilerplate.js
+++ b/Engine/Extras/RoboMerge/v3/public/js/boilerplate.js
@@ -1592,11 +1592,9 @@ function renderActionsCell_Edge(actionCell, nodeData, edgeData, conflict=null) {
 		}
 	}
 
-	let tooManyFilesBlockage = conflict && conflict.kind === 'Too many files'
-
 	if (edgeData.is_blocked && edgeData.blockage && edgeData.blockage.targetBranchName) {
 		// Skip CL
-		const skipEnabled = tooManyFilesBlockage || !edgeData.disallowSkip
+		const skipEnabled = !edgeData.disallowSkip
 		
 		const skipChangelistText = `Skip Changelist ${edgeData.blockage.change}`
 		if (skipEnabled) {
@@ -1607,10 +1605,8 @@ function renderActionsCell_Edge(actionCell, nodeData, edgeData, conflict=null) {
 				edge: edgeData.blockage.targetBranchName
 			}) + location.hash
 
-			let tooltip = `Skip past the blockage caused by changelist ${edgeData.blockage.change}. `
-			tooltip += 
-				tooManyFilesBlockage ?  'Please ensure this large changelist has been integrated before skipping.' :
-				"This option should only be selected if the work does not need to be merged or you will merge this work youself."
+			const tooltip = `Skip past the blockage caused by changelist ${edgeData.blockage.change}. `
+				+ "This option should only be selected if the work does not need to be merged or you will merge this work youself.`
 			const skipOption = createActionOption(skipChangelistText, function() {
 				window.location.href = skipRequest;
 			}, tooltip)
diff --git a/Engine/Extras/RoboMerge/v3/src/robo/notifications.ts b/Engine/Extras/RoboMerge/v3/src/robo/notifications.ts
--- a/Engine/Extras/RoboMerge/v3/src/robo/notifications.ts
+++ b/Engine/Extras/RoboMerge/v3/src/robo/notifications.ts
@@ -507,7 +507,6 @@ export class BotNotifications implements BotEventHandler {
 		const text =
 			blockage.approval ?				`${channelPing}'s change needs to be approved in <#${blockage.approval.settings.channelId}>` :
 			isBotUser ? 									`Blockage caused by \`${blockage.owner}\` commit!` :
-			blockage.failure.kind === 'Too many files' ?	`${channelPing}, please request a shelf for this large changelist` :
 															`${channelPing}, please resolve the following ${issue}:`
 
 		const message = this.makeSlackChannelMessage(
diff --git a/Engine/Extras/RoboMerge/v3/src/robo/status-types.ts b/Engine/Extras/RoboMerge/v3/src/robo/status-types.ts
--- a/Engine/Extras/RoboMerge/v3/src/robo/status-types.ts
+++ b/Engine/Extras/RoboMerge/v3/src/robo/status-types.ts
@@ -23,7 +23,6 @@ export type FailureKind =
 	'Commit failure' |
 	'Syntax error' |
 	'Disallowed files' |
-	'Too many files' |
 	'Approval required' |
 	'Conversion to edits failure' |
 	'Unit Test error'
