Fixed incorrect string closure causing parse error

diff --git a/Engine/Extras/RoboMerge/v3/public/js/boilerplate.js b/Engine/Extras/RoboMerge/v3/public/js/boilerplate.js
--- a/Engine/Extras/RoboMerge/v3/public/js/boilerplate.js
+++ b/Engine/Extras/RoboMerge/v3/public/js/boilerplate.js
@@ -1606,7 +1606,7 @@ function renderActionsCell_Edge(actionCell, nodeData, edgeData, conflict=null) {
 			}) + location.hash
 
 			const tooltip = `Skip past the blockage caused by changelist ${edgeData.blockage.change}. `
-				+ "This option should only be selected if the work does not need to be merged or you will merge this work youself.`
+				+ "This option should only be selected if the work does not need to be merged or you will merge this work youself."
 			const skipOption = createActionOption(skipChangelistText, function() {
 				window.location.href = skipRequest;
 			}, tooltip)
