Avoid crashing robomerge when tracking extremely large changes

diff --git a/Engine/Extras/RoboMerge/v3/src/robo/ipc.ts b/Engine/Extras/RoboMerge/v3/src/robo/ipc.ts
--- a/Engine/Extras/RoboMerge/v3/src/robo/ipc.ts
+++ b/Engine/Extras/RoboMerge/v3/src/robo/ipc.ts
@@ -349,8 +349,9 @@ export class IPC {
 							(RobomergeMethodStrings as readonly string[]).includes(mergeMethod)) {
 					// If we only have 1 entry and we didn't get integration info off of it
 					// and the graph suggests we are expecting there could be other changes
-					// get the full describe results
-					changeToConsider.desc = await this.robo.p4.describe(clToConsider)
+					// get more of the describe results. We'll limit ourselves somewhat to avoid
+					// overflowing the buffer
+					changeToConsider.desc = await this.robo.p4.describe(clToConsider, 1000)
 				}
 			}
 			clsToConsider = clsToConsider.concat(changeToConsider.destCLs)
