Add dirs and files perforce functions
Try to use a more sophisticated attempt at finding additional files to track when the first file from describe wasn't found

diff --git a/Engine/Extras/RoboMerge/v3/src/common/perforce.ts b/Engine/Extras/RoboMerge/v3/src/common/perforce.ts
--- a/Engine/Extras/RoboMerge/v3/src/common/perforce.ts
+++ b/Engine/Extras/RoboMerge/v3/src/common/perforce.ts
@@ -1430,6 +1430,20 @@ export class PerforceContext {
 		return result;
 	}
 
+	async dirs(path: string) {
+		if (path.endsWith("/...")) {
+			path = path.substring(0,path.length-3) + "*"
+		}
+		return (await this._execP4Ztag(null, ['dirs', path], { multiline:true })).map((dir: any) => dir.dir)
+	}
+
+	async files(path: string, maxFiles?: number) {
+		const args = ['files',
+			...(maxFiles ? ['-m', maxFiles.toString()] : []),
+			path]
+		return await this._execP4Ztag(null, args, { multiline:true });
+	}
+
 	// update the fields on an existing CL using p4 change
 	// output format is just error or not
 	async editChange(roboWorkspace: RoboWorkspace, changelist: number, opts?: EditChangeOpts) {
diff --git a/Engine/Extras/RoboMerge/v3/src/robo/ipc.ts b/Engine/Extras/RoboMerge/v3/src/robo/ipc.ts
--- a/Engine/Extras/RoboMerge/v3/src/robo/ipc.ts
+++ b/Engine/Extras/RoboMerge/v3/src/robo/ipc.ts
@@ -349,9 +349,17 @@ export class IPC {
 							(RobomergeMethodStrings as readonly string[]).includes(mergeMethod)) {
 					// If we only have 1 entry and we didn't get integration info off of it
 					// and the graph suggests we are expecting there could be other changes
-					// get more of the describe results. We'll limit ourselves somewhat to avoid
-					// overflowing the buffer
-					changeToConsider.desc = await this.robo.p4.describe(clToConsider, 1000)
+					// get more of the files. If the path is a root directory try and get a sampling 
+					// of files, otherwise just get an additional block to evaluate
+					if (changeToConsider.desc.path.endsWith("/...")) {
+						const dirs = await this.robo.p4.dirs(changeToConsider.desc.path)
+						for (const dir of dirs) {
+							changeToConsider.desc.entries.push(...(await this.robo.p4.files(`${dir}/...@=${clToConsider}`, 1)))
+						}
+					}
+					else {
+						changeToConsider.desc = await this.robo.p4.describe(clToConsider, 100)
+					}
 				}
 			}
 			clsToConsider = clsToConsider.concat(changeToConsider.destCLs)
