Introduce uniqueBranch identifier to indicate that a node should not alias another just because it shares a rootPath
Set UE5 branchmap to use the new setting

diff --git a/Engine/Extras/RoboMerge/v3/src/new/graph.ts b/Engine/Extras/RoboMerge/v3/src/new/graph.ts
--- a/Engine/Extras/RoboMerge/v3/src/new/graph.ts
+++ b/Engine/Extras/RoboMerge/v3/src/new/graph.ts
@@ -54,7 +54,10 @@ export function makeTargetName(bot: BotName, nodeName: string) {
 
 function makeBranchId(arg: string | Branch) {
 	let depotPath: string
-	if ((arg as Branch).rootPath) {
+	if ((arg as Branch).uniqueBranch) {
+		return (arg as Branch).name as Stream
+	}
+	else if ((arg as Branch).rootPath) {
 		depotPath = (arg as Branch).rootPath.replace('/...', '')
 	}
 	else {
@@ -369,7 +372,6 @@ export function addBranchGraph(graph: Graph, branchGraph: BranchGraphInterface)
 	const botname = branchGraph.botname as BotName
 	const branchNodes = new Map<Branch, Node>()
 
-	// excluding subpath bots for now
 	for (const branch of branchGraph.branches) {
 		if (!branch.bot) {
 			throw new Error(`branch ${branch.name} not running!`) // fine, but try again later
diff --git a/Engine/Extras/RoboMerge/v3/src/robo/branchdefs.ts b/Engine/Extras/RoboMerge/v3/src/robo/branchdefs.ts
--- a/Engine/Extras/RoboMerge/v3/src/robo/branchdefs.ts
+++ b/Engine/Extras/RoboMerge/v3/src/robo/branchdefs.ts
@@ -41,6 +41,7 @@ const branchBasePrototype = {
 	name: '',
 
 	rootPath: '',
+	uniqueBranch: false,
 	isDefaultBot: false,
 	emailOnBlockage: false, // if present, completely overrides BotConfig
 
@@ -122,6 +123,7 @@ const nodeOptionFieldsPrototype = {
 	streamDepot: '',
 	streamName: '',
 	streamSubpath: '',
+	uniqueBranch: false,
 	workspace: '',
 
 	graphNodeColor: '',
diff --git a/Engine/Extras/RoboMerge/v3/src/robo/branchgraph.ts b/Engine/Extras/RoboMerge/v3/src/robo/branchgraph.ts
--- a/Engine/Extras/RoboMerge/v3/src/robo/branchgraph.ts
+++ b/Engine/Extras/RoboMerge/v3/src/robo/branchgraph.ts
@@ -228,6 +228,7 @@ if (botname === '__TEST__') {
 			config: options,
 			depot: "", // will compute
 			rootPath: options.rootPath || "",
+			uniqueBranch: options.uniqueBranch || false,
 			badgeProject: options.badgeProject || null,
 			isDefaultBot: BranchGraph._getBoolConfig(options.isDefaultBot, this.config.isDefaultBot),
 			emailOnBlockage: BranchGraph._getBoolConfig(options.emailOnBlockage, this.config.emailOnBlockage),
