Fix stream name determination when the supplied path is exactly the stream root

diff --git a/Engine/Extras/RoboMerge/v3/src/common/perforce.ts b/Engine/Extras/RoboMerge/v3/src/common/perforce.ts
--- a/Engine/Extras/RoboMerge/v3/src/common/perforce.ts
+++ b/Engine/Extras/RoboMerge/v3/src/common/perforce.ts
@@ -602,6 +602,15 @@ export class PerforceContext {
 			}
 			let streamNameEnd = 2;
 			for (let i=0;i<streamDepth+1;i++) {
+				let nextSlash = path.indexOf("/",streamNameEnd+1)
+				if (nextSlash == -1) {
+					if (i == streamDepth) {
+						return path
+					}
+					else {
+						return new Error(`Unable to determine stream from $(path): not enough depth`)
+					}
+				}
 				streamNameEnd = path.indexOf("/",streamNameEnd+1)
 			}
 			return path.substring(0,streamNameEnd)
