Fix gate tests

diff --git a/Engine/Extras/RoboMerge/v3/src/robo/gate.ts b/Engine/Extras/RoboMerge/v3/src/robo/gate.ts
--- a/Engine/Extras/RoboMerge/v3/src/robo/gate.ts
+++ b/Engine/Extras/RoboMerge/v3/src/robo/gate.ts
@@ -797,13 +797,13 @@ export async function runTests(parentLogger: ContextualLogger) {
 		const [et, gate] = makeTestGate(1, options)
 		await gate.tick()
 		assert(et.beginCalls + et.endCalls === 0, 'no events yet')
-		assert(!gate.isGateOpen() && gate.getGateClosedMessage()!.includes('CIS'), 'gate closed')
+		assert(!gate.isGateOpen() && gate.getGateClosedMessage()! == 'paused at gate', 'gate closed')
 
 		// nothing should happen here - CL 2 has been committed but gate prevents it being integrated
 		gate.preIntegrate(2)
 		await gate.tick()
 		assert(et.beginCalls + et.endCalls === 0, 'no events yet')
-		assert(!gate.isGateOpen() && gate.getGateClosedMessage()!.includes('CIS'), 'gate closed')
+		assert(!gate.isGateOpen() && gate.getGateClosedMessage()! == 'paused at gate', 'gate closed')
 
 		// gate is now > 1 so we can 
 		options.lastGoodCLPath = exact ? 2 : 3
@@ -885,7 +885,7 @@ export async function runTests(parentLogger: ContextualLogger) {
 		}
 
 		assert(et.beginCalls === 1 && et.endCalls === 1, 'caught up')
-		assert(!gate.isGateOpen() && gate.getGateClosedMessage()!.includes('CIS'), 'gate closed')
+		assert(!gate.isGateOpen() && gate.getGateClosedMessage()! == 'paused at gate', 'gate closed')
 	}
 
 	const queueNoWindow = async (middleIntegration: boolean) => {
@@ -992,7 +992,7 @@ export async function runTests(parentLogger: ContextualLogger) {
 		gate.preIntegrate(7)
 		setLastCl(gate, 7)
 		assert(et.beginCalls === 1 && et.endCalls === 1, 'caught up')
-		assert(!gate.isGateOpen() && gate.getGateClosedMessage()!.includes('CIS'), 'gate closed')
+		assert(!gate.isGateOpen() && gate.getGateClosedMessage()! == 'paused at gate', 'gate closed')
 
 	})()
 
