Post locked files and merge conflicts as reply using snippets
Clean up dummyslack server to cover all current used calls during functional tests

diff --git a/Engine/Extras/RoboMerge/v3/src/common/dummyslackserver.ts b/Engine/Extras/RoboMerge/v3/src/common/dummyslackserver.ts
--- a/Engine/Extras/RoboMerge/v3/src/common/dummyslackserver.ts
+++ b/Engine/Extras/RoboMerge/v3/src/common/dummyslackserver.ts
@@ -44,20 +44,8 @@ export class DummySlackApp implements AppInterface {
 			return { dataObj }
 		}
 
-		if (command === "conversations.invite") {
-			return { ok: true }
-		}
-		else if (command === "conversations.open") {
-			const dataResult = getData()
-			if (dataResult.error) {
-				return dataResult
-			}
-			const dataObj = dataResult.dataObj
-			return { ok: true, channel: {id: `${dataObj.users}`} }
-		}
-		else if (command === "chat.update") {
-			// Not going to handle updating the stored messages on the dummy server
-			// until we have a reason to do so 
+		if (command === "chat.getPermalink") {
+			// We don't need the permalink for any functional tests
 			return { ok: true }
 		}
 		else if (command === "chat.postMessage") {
@@ -101,6 +89,26 @@ export class DummySlackApp implements AppInterface {
 					return { error: "cl or edge not found"}
 				}
 			}
+		}		
+		else if (command === "chat.update") {
+			// Not going to handle updating the stored messages on the dummy server
+			// until we have a reason to do so 
+			return { ok: true }
+		}
+		else if (command === "conversations.invite") {
+			return { ok: true }
+		}
+		else if (command === "conversations.open") {
+			const dataResult = getData()
+			if (dataResult.error) {
+				return dataResult
+			}
+			const dataObj = dataResult.dataObj
+			return { ok: true, channel: {id: `${dataObj.users}`} }
+		}
+		else if (command === "files.upload") {
+			// Not going to store files until we have a reason to do so
+			return { ok: true }
 		}
 		else {
 			console.error(`dummy-slack POST: unsupported command '${command}'`)
diff --git a/Engine/Extras/RoboMerge/v3/src/robo/branch-interfaces.ts b/Engine/Extras/RoboMerge/v3/src/robo/branch-interfaces.ts
--- a/Engine/Extras/RoboMerge/v3/src/robo/branch-interfaces.ts
+++ b/Engine/Extras/RoboMerge/v3/src/robo/branch-interfaces.ts
@@ -181,7 +181,7 @@ export interface PendingChange {
 export interface Failure {
 	kind: FailureKind		// short description of integration error or conflict
 	description: string		// detailed description (can be very long - don't want to store this)
-	summary?: string
+	details?: string
 	additionalInfo?: any
 }
 
diff --git a/Engine/Extras/RoboMerge/v3/src/robo/edgebot.ts b/Engine/Extras/RoboMerge/v3/src/robo/edgebot.ts
--- a/Engine/Extras/RoboMerge/v3/src/robo/edgebot.ts
+++ b/Engine/Extras/RoboMerge/v3/src/robo/edgebot.ts
@@ -23,13 +23,9 @@ import { BlockagePauseInfo, BlockagePauseInfoMinimal, EdgeStatusFields } from ".
 import { getIntegrationOwner } from "./targets";
 
 const FAILED_CHANGELIST_PAUSE_TIMEOUT_SECONDS = 15 * 60
-const MAX_INTEGRATION_ERRORS_TO_ANALYZE = 5
 const MAX_OPENED_COMMANDS_TO_ISSUE = 100
-const DEPOT_FILE_REGEX = /^(.*[\\\/])(.*)/
 const JIRA_REGEX = /^\s*#jira\s+(.*)/i
 
-const MAX_CONFLICTS_TO_LIST = 5
-
 // type ResolveResultDetail = 'quick' | 'detailed'
 
 type EdgeIntegrationResult = 'ok' | 'shelved' | 'nothing to do' | 'skipped' | 'error'
@@ -241,20 +237,17 @@ class EdgeBotImpl extends PerforceStatefulBot {
 		const results: ConflictingFile[] = []
 
 		for (const file of unresolved) {
-			const match = file.fromFile.match(DEPOT_FILE_REGEX)
-			if (match) {
-				if (file.resolveType.toLowerCase() === "branch") {
-					results.push({name: match[2], kind: "branch"})
-				}
-				else if (file.resolveType.toLowerCase() === "delete") {
-					results.push({name: match[2], kind: "delete"})
-				}
-				else if (file.resolveType.toLowerCase() === "content") {
-					results.push({name: match[2], kind: "merge"})
-				} else {
-					// We really shouldn't get this kind, but it's better to display unknown than skip displaying the file
-					results.push({name: match[2], kind: "unknown"})
-				}
+			if (file.resolveType.toLowerCase() === "branch") {
+				results.push({name: file.fromFile, kind: "branch"})
+			}
+			else if (file.resolveType.toLowerCase() === "delete") {
+				results.push({name: file.fromFile, kind: "delete"})
+			}
+			else if (file.resolveType.toLowerCase() === "content") {
+				results.push({name: file.fromFile, kind: "merge"})
+			} else {
+				// We really shouldn't get this kind, but it's better to display unknown than skip displaying the file
+				results.push({name: file.fromFile, kind: "unknown"})
 			}
 		}
 		return results
@@ -562,12 +555,8 @@ class EdgeBotImpl extends PerforceStatefulBot {
 
 		if (exclusiveFiles.length > 0) {
 			// will need to store the exclusive file if we want to @ people in Slack
-			const exclCheckoutMessages = exclusiveFiles.map(exc => `${exc.depotPath} checked out by ${exc.user}`)
-			if (errors.length > MAX_INTEGRATION_ERRORS_TO_ANALYZE) {
-				exclCheckoutMessages.push(`... and ${errors.length - MAX_INTEGRATION_ERRORS_TO_ANALYZE} more`)
-			}
 			const exclusiveLockUsers = Array.from(new Set(exclusiveFiles.map(exc => `${exc.user.toLowerCase()}`))).map(user => ({user, userEmail: this.p4.getEmail(user)}))
-			failure = { kind: 'Exclusive check-out', description, summary: exclCheckoutMessages.join('\n'), additionalInfo: {exclusiveLockUsers,exclusiveFiles} }
+			failure = { kind: 'Exclusive check-out', description, additionalInfo: {exclusiveLockUsers,exclusiveFiles} }
 		}
 		else {
 			failure  = { kind: 'Integration error', description }
@@ -648,14 +637,9 @@ class EdgeBotImpl extends PerforceStatefulBot {
 
 			const conflicts = this.analyzeConflict(result.getConflicts())
 			if (conflicts.length > 0) {
-				failure.summary = conflicts
-					.slice(0, MAX_CONFLICTS_TO_LIST)
+				failure.details = conflicts
 					.map(({name, kind}) => `${name} (${kind} conflict)`)
 					.join('\n')
-
-				if (conflicts.length > MAX_CONFLICTS_TO_LIST) {
-					failure.summary += `\n... and ${conflicts.length - MAX_CONFLICTS_TO_LIST} more`
-				}
 			}
 		}
 		else
@@ -690,14 +674,14 @@ class EdgeBotImpl extends PerforceStatefulBot {
 					return result
 				}
 
-				let summary: string | undefined
+				let details: string | undefined
 				const match = result.message.match(/.*STDERR:([^]*)STDOUT:/)
 				if (match)
 				{
-					summary = match[1].trim()
+					details = match[1].trim()
 				}
 
-				failure = { kind: 'Commit failure', description: result.message, summary }
+				failure = { kind: 'Commit failure', description: result.message, details }
 			}
 		}
 
diff --git a/Engine/Extras/RoboMerge/v3/src/robo/nodebot.ts b/Engine/Extras/RoboMerge/v3/src/robo/nodebot.ts
--- a/Engine/Extras/RoboMerge/v3/src/robo/nodebot.ts
+++ b/Engine/Extras/RoboMerge/v3/src/robo/nodebot.ts
@@ -18,7 +18,7 @@ import { makeClLink, SlackMessages } from './notifications';
 import { PerforceStatefulBot } from './perforce-stateful-bot';
 import { BlockageNodeOpUrls, OperationUrlHelper } from './roboserver';
 import { Context } from './settings';
-import { SlackMessage, SlackMessageField, SlackMessageStyles } from './slack';
+import { SlackFile, SlackMessage, SlackMessageField, SlackMessageStyles } from './slack';
 import { PauseState } from './state-interfaces';
 import { newTickJournal, TickJournal } from './tick-journal';
 import { computeTargets, parseDescriptionLines, processOtherBotTargets, getIntegrationOwner, getNodeBotFullName, getNodeBotFullNameForLogging } from './targets';
@@ -2118,8 +2118,8 @@ export class NodeBot extends PerforceStatefulBot implements NodeBotInterface {
 
 			let message = `While reconsidering CL#${pending.change.cl}`
 
-			if (failure.summary) {
-				message += '\n' + failure.summary 
+			if (failure.details) {
+				message += '\n' + failure.details 
 			}
 			else
 			{
@@ -2141,6 +2141,15 @@ export class NodeBot extends PerforceStatefulBot implements NodeBotInterface {
 				
 				let emailAddress = await this.findEmail(owner)
 				this.slackMessages.postDM(emailAddress, pending.newCl, this.branch, dm)
+
+				if (failure.details) {
+					let file: SlackFile = {
+						content: failure.details,
+						channels: "",
+						filename: "failuredetails.txt"
+					}
+					this.slackMessages.postFileToDM(emailAddress, pending.newCl, this.branch, file)
+				}				
 			}
 			else {
 				this.sendErrorEmail(new Recipients(owner), shortMessage, message)
diff --git a/Engine/Extras/RoboMerge/v3/src/robo/notifications.ts b/Engine/Extras/RoboMerge/v3/src/robo/notifications.ts
--- a/Engine/Extras/RoboMerge/v3/src/robo/notifications.ts
+++ b/Engine/Extras/RoboMerge/v3/src/robo/notifications.ts
@@ -4,15 +4,15 @@ import * as fs from 'fs';
 import { DateTimeFormatOptions } from 'intl';
 import { Args } from '../common/args';
 import { Badge } from '../common/badge';
-import { Random } from '../common/helper';
+import { Random, setDefault } from '../common/helper';
 import { ContextualLogger } from '../common/logger';
-import { Blockage, Branch, BranchArg, ExclusiveLockInfo, ForcedCl, MergeAction, NodeOpUrlGenerator, resolveBranchArg } from './branch-interfaces';
+import { Blockage, Branch, BranchArg, ExclusiveFile, ExclusiveLockInfo, ForcedCl, MergeAction, NodeOpUrlGenerator, resolveBranchArg } from './branch-interfaces';
 import { PersistentConflict, Resolution } from './conflict-interfaces';
 import { BotEventHandler, BotEvents } from './events';
 import { NodeBot } from './nodebot';
 import { BlockageNodeOpUrls } from './roboserver';
 import { Context } from './settings';
-import { Slack, SlackAttachment, SlackLinkButtonAction, SlackLinkButtonsAttachment, SlackMessageField, SlackMessage, SlackMessageStyles } from './slack';
+import { Slack, SlackAttachment, SlackFile, SlackLinkButtonAction, SlackLinkButtonsAttachment, SlackMessageField, SlackMessage, SlackMessageStyles } from './slack';
 import { WebServer } from '../common/webserver'
 import { DummySlackApp } from '../common/dummyslackserver'
 
@@ -290,14 +290,32 @@ export class SlackMessages {
 		}
 	}
 
-	async postDM(emailAddress: string|Promise<string|null>|null, cl: number, branchArg: BranchArg, dm: SlackMessage, persistMessage = true) {
+	async postFile(cl: number, branchArg: BranchArg, file: SlackFile) {
+		const findResult = this.find(cl, branchArg, file.channels)
+
+		if (findResult.messageRecord) {
+			try {
+				file.thread_ts = findResult.messageRecord.timestamp
+				await this.slack.uploadFile(file)
+			}
+			catch (err) {
+				console.error('Error uploading file to Slack! ' + err.toString())
+				return
+			}
+		}
+		else {
+			this.smLogger.error(`Failed to find message record for ${cl}:${branchArg}:${file.channels}`)
+		}
+	}
+
+	private async getDMChannelId(emailAddress: string|Promise<string|null>|null, cl: number) {
 		// The Slack API requires a user ID to open a direct message with users.
 		// The most consistent way to do this is getting their email address out of P4.
 		if (emailAddress && typeof(emailAddress) !== 'string') {
 			emailAddress = await emailAddress
 		}
 		if (!emailAddress) {
-			console.error("Failed to get email address during notifications for CL " + cl)
+			this.smLogger.error("Failed to get email address during notifications for CL " + cl)
 			return
 		}
 
@@ -310,18 +328,34 @@ export class SlackMessages {
 		}
 
 		// Open up a new conversation with the user now that we have their ID
-		let channelId : string
 		try {
-			channelId = (await this.slack.openDMConversation(userId))
-			dm.channel = channelId
+			return this.slack.openDMConversation(userId)
 		} catch (err) {
 			this.smLogger.printException(err, `Failed to get Slack conversation ID for user ID "${userId}" given email address "${emailAddress}" for CL ${cl}`)
-			return
 		}
 
-		// Add the channel/conversation ID to the messageOpts and proceed normally.
-		this.smLogger.info(`Creating direct message for ${emailAddress} (key: ${generatePersistedSlackMessageKey(cl, branchArg, channelId)})`)
-		this.postOrUpdate(cl, branchArg, dm, persistMessage)
+		return
+	}
+
+	async postDM(emailAddress: string|Promise<string|null>|null, cl: number, branchArg: BranchArg, dm: SlackMessage, persistMessage = true) {
+
+		const channelId = await this.getDMChannelId(emailAddress, cl)
+		if (channelId) {
+			dm.channel = channelId
+			// Add the channel/conversation ID to the messageOpts and proceed normally.
+			this.smLogger.info(`Creating direct message for ${emailAddress} (key: ${generatePersistedSlackMessageKey(cl, branchArg, channelId)})`)
+			this.postOrUpdate(cl, branchArg, dm, persistMessage)
+		}
+	}
+
+	async postFileToDM(emailAddress: string|Promise<string|null>|null, cl: number, branchArg: BranchArg, file: SlackFile) {
+
+		const channelId = await this.getDMChannelId(emailAddress, cl)
+		if (!channelId) {
+			return
+		}
+		file.channels = channelId
+		this.postFile(cl, branchArg, file)
 	}
 
 	// 
@@ -519,10 +553,6 @@ export class BotNotifications implements BotEventHandler {
 			changeInfo.author
 		);
 
-		if (blockage.failure.summary) {
-			message.footer = blockage.failure.summary
-		}
-
 		let messagesToPost = []
 		let channelsToPostTo = []
 		let usersToInvite: Set<string> = new Set()
@@ -546,12 +576,29 @@ export class BotNotifications implements BotEventHandler {
 
 			messagesToPost.push({ message })
 
+			if (isNew && blockage.failure.details) {
+				let file: SlackFile = {
+					content: blockage.failure.details,
+					channels: message.channel,
+					filename: "conflictdetails.txt"
+				}
+				messagesToPost.push({file})
+			}
+
 			if (isNew && blockage.failure.kind === 'Exclusive check-out') {
-				const exclusiveLockUsers = (blockage.failure.additionalInfo as ExclusiveLockInfo).exclusiveLockUsers
-				let text = ''
+				const exclusiveLockInfo = blockage.failure.additionalInfo as ExclusiveLockInfo
+
+				let authorDict = new Map<string, ExclusiveFile[]>()
+				for (const file of exclusiveLockInfo.exclusiveFiles) {
+					setDefault(authorDict, file.user.toLowerCase(), []).push(file)
+				}
 
-				let unknownLockUsers = false
-				for (const exclusiveLockUser of exclusiveLockUsers) {
+				for (const exclusiveLockUser of exclusiveLockInfo.exclusiveLockUsers) {
+					let file: SlackFile = { 
+						content: authorDict.get(exclusiveLockUser.user)!.map(ef => ef.depotPath).join("\n"),
+						channels: message.channel,
+						filename: "exclusivelockedfiles.txt"
+					}
 					if (exclusiveLockUser.user.length > 0) {
 						let exclusiveLockSlackUser
 						const exclusiveLockUserEmail = await exclusiveLockUser.userEmail
@@ -565,23 +612,12 @@ export class BotNotifications implements BotEventHandler {
 						if (!exclusiveLockSlackUser) {
 							exclusiveLockSlackUser = `@${exclusiveLockUser.user} `
 						}
-						text += exclusiveLockSlackUser
+						file.initial_comment = `${exclusiveLockSlackUser} please unlock the files blocking robomerge or work with ${blockage.owner} to resolve the conflict`
 					} else {
-						unknownLockUsers = true
+						file.initial_comment = "The following locked files did not have their owner determined and as such those owners may not have been notified"
 					}
+					messagesToPost.push({file})
 				}
-				text += `\n\nPlease unlock the files blocking robomerge or work with ${blockage.owner} to resolve the conflict`
-				if (unknownLockUsers) {
-					text += "\n\nNote that some locked files did not have their owner determined and as such those owners may not have been tagged"
-				}
-				
-				const exclusiveCheckoutMessage: SlackMessage = {
-					text,
-					style: SlackMessageStyles.DANGER,
-					channel: message.channel,
-					mrkdwn: true
-				}
-				messagesToPost.push({ message: exclusiveCheckoutMessage, reply: true})
 			}
 		}
 		else {
@@ -603,7 +639,10 @@ export class BotNotifications implements BotEventHandler {
 				this.slackMessages.addUserToChannel(user, channel)
 			}
 			for (const messageToPost of messagesToPost) {
-				if (messageToPost.reply) {
+				if (messageToPost.file) {
+					await this.slackMessages.postFile(changeInfo.cl, branchArg, { ...messageToPost.file, channels: channel})
+				}
+				else if (messageToPost.reply) {
 					await this.slackMessages.postReply(changeInfo.cl, branchArg, { ...messageToPost.message, channel })
 				}
 				else {
@@ -740,7 +779,7 @@ export class BotNotifications implements BotEventHandler {
 					"0",
 					false
 				)
-			)
+		)
 
 			this.slackMessages.postDM([REDACTED_EMAIL]`, 0, "TARGETBRANCH", messageOpts, false)
 		}
diff --git a/Engine/Extras/RoboMerge/v3/src/robo/slack.ts b/Engine/Extras/RoboMerge/v3/src/robo/slack.ts
--- a/Engine/Extras/RoboMerge/v3/src/robo/slack.ts
+++ b/Engine/Extras/RoboMerge/v3/src/robo/slack.ts
@@ -2,6 +2,7 @@
 
 import * as request from '../common/request';
 import { ContextualLogger } from '../common/logger';
+import { URLSearchParams } from "url"
 
 export interface SlackChannel {
 	id: string
@@ -68,6 +69,20 @@ export interface SlackMessage {
 	target?: string
 }
 
+export interface SlackFile {
+	content: string
+	channels: string
+	filename?: string
+	filetype?: string
+	initial_comment?: string
+	thread_ts?: string
+	title?: string
+}
+
+enum PostFormat {
+	JSON,
+	URL
+}
 
 const MAIN_MESSAGE_FIELDS = new Set(['username', 'icon_emoji', 'channel', 'target', 'cl']);
 
@@ -106,16 +121,19 @@ export class Slack {
 		return this.post('chat.postMessage', args)
 	}
 
-// why was message optional?
 	update(ts: string, message: SlackMessage) {
 		const args = this.makeArgs(message)
 		args.ts = ts
 		return this.post('chat.update', args)
 	}
 
+	uploadFile(file: SlackFile) {
+		return this.post('files.upload', file, false, PostFormat.URL)
+	}
+
 	listMessages(count?: number) {
 		const args: any = count ? {count} : count
-		// use channels.history if publlic?
+		// use channels.history if public?
 		return this.get('groups.history', args)
 	}
 
@@ -131,12 +149,12 @@ export class Slack {
 		return (await this.post('conversations.open', {users})).channel.id
 	}
 
-	/*private*/ async post_user(userToken: string, command: string, args: any, canFail? : boolean) {
+	/*private*/ async post_user(userToken: string, command: string, args: any, canFail? : boolean, format?: PostFormat) {
 		const resultJson = await request.post({
 			url: this.domain + '/api/' + command,
-			body: JSON.stringify(args),
+			body: format == PostFormat.URL ? new URLSearchParams(Object.entries(args)).toString() : JSON.stringify(args),
 			headers: {Authorization: 'Bearer ' + userToken},
-			contentType: 'application/json; charset=utf-8'
+			contentType: format == PostFormat.URL ? 'application/x-www-form-urlencoded' : 'application/json; charset=utf-8'
 		})
 		try {
 			const result = JSON.parse(resultJson)
@@ -151,8 +169,8 @@ export class Slack {
 		return {ok: false}
 	}
 
-	/*private*/ async post(command: string, args: any, canFail? : boolean) {
-		return this.post_user(this.channel.botToken, command, args, canFail)
+	/*private*/ async post(command: string, args: any, canFail? : boolean, format?: PostFormat) {
+		return this.post_user(this.channel.botToken, command, args, canFail, format)
 	}
 
 	/*private*/ async get(command: string, args: any, canFail? : boolean) {
