Patch crash when openconversation fails

diff --git a/Engine/Extras/RoboMerge/v3/src/robo/slack.ts b/Engine/Extras/RoboMerge/v3/src/robo/slack.ts
--- a/Engine/Extras/RoboMerge/v3/src/robo/slack.ts
+++ b/Engine/Extras/RoboMerge/v3/src/robo/slack.ts
@@ -142,11 +142,12 @@ export class Slack {
 		return userLookupResult.ok ? userLookupResult.user.id : null
 	}
 
-	async openDMConversation(users: string | string[]) : Promise<string> {
+	async openDMConversation(users: string | string[]) : Promise<string|null> {
 		if (users instanceof Array) {
 			users = users.join(',')
 		}
-		return (await this.post('conversations.open', {users})).channel.id
+		const result = await this.post('conversations.open', {users})
+		return result.channel ? result.channel.id : null
 	}
 
 	/*private*/ async post_user(userToken: string, command: string, args: any, canFail? : boolean, format?: PostFormat) {
