Fix track changes incorrectly identifying that a node has an automerge target
Fix track changes calling files at a specific CL on paths that don't have any changes in that path

diff --git a/Engine/Extras/RoboMerge/v3/src/robo/ipc.ts b/Engine/Extras/RoboMerge/v3/src/robo/ipc.ts
--- a/Engine/Extras/RoboMerge/v3/src/robo/ipc.ts
+++ b/Engine/Extras/RoboMerge/v3/src/robo/ipc.ts
@@ -291,7 +291,7 @@ export class IPC {
 							}
 						}
 					}
-					if (edge.flags.has('automatic')) {
+					if (!hasAutomergeTarget && edge.flags.has('automatic') && edge.source.stream == changeToConsider.node) {
 						hasAutomergeTarget = true
 					}
 					if (includeInResults && hasAutomergeTarget) {
@@ -352,7 +352,7 @@ export class IPC {
 					// get more of the files. If the path is a root directory try and get a sampling 
 					// of files, otherwise just get an additional block to evaluate
 					if (changeToConsider.desc.path.endsWith("/...")) {
-						const dirs = await this.robo.p4.dirs(changeToConsider.desc.path)
+						const dirs = await this.robo.p4.dirs(`${changeToConsider.desc.path.slice(0,-3)}*@=${clToConsider}`)
 						for (const dir of dirs) {
 							changeToConsider.desc.entries.push(...(await this.robo.p4.files(`${dir}/...@=${clToConsider}`, 1)))
 						}
