Fix quiet parameter to changes not being handled correctly

diff --git a/Engine/Extras/RoboMerge/v3/src/common/perforce.ts b/Engine/Extras/RoboMerge/v3/src/common/perforce.ts
--- a/Engine/Extras/RoboMerge/v3/src/common/perforce.ts
+++ b/Engine/Extras/RoboMerge/v3/src/common/perforce.ts
@@ -540,14 +540,14 @@ export class PerforceContext {
 	 * Get a list of changes in a path since a specific CL
 	 * @return Promise to list of changelists
 	 */
-	changes(path_in: string, since: number, limit?: number, status?: ChangelistStatus, quiet?: boolean): Promise<Change[]> {
+	changes(path_in: string, since: number, limit?: number, status?: ChangelistStatus, quiet: boolean = true): Promise<Change[]> {
 		const path = since > 0 ? path_in + '@>' + since : path_in;
 		const args = ['changes', '-l',
 			(status ? `-s${status}` : '-ssubmitted'),
 			...(limit ? [`-m${limit}`] : []),
 			path];
 
-		return this.execAndParse(null, args, {quiet: quiet ? quiet : true}, {
+		return this.execAndParse(null, args, {quiet}, {
 			expected: {change: 'integer', client: 'string', user: 'string', desc: 'string'},
 			optional: {shelved: 'integer', oldChange: 'integer', IsPromoted: 'integer'}
 		}) as Promise<unknown> as Promise<Change[]>
