[NO NON-PII COMMIT MESSAGE]

diff --git a/Engine/Extras/RoboMerge/v3/public/js/boilerplate.js b/Engine/Extras/RoboMerge/v3/public/js/boilerplate.js
--- a/Engine/Extras/RoboMerge/v3/public/js/boilerplate.js
+++ b/Engine/Extras/RoboMerge/v3/public/js/boilerplate.js
@@ -1031,7 +1031,7 @@ function createPauseDivs(data, conflict) {
 			data.blockage ? `<span class="pause-div-label">Blocked.</span> Type: ${data.blockage.type}<br /> Message: ${data.blockage.message}` :
 			`No info can be provided. Please contact Robomerge help.`
 
-		if (conflict.slackLinks) {
+		if (conflict && conflict.slackLinks) {
 			info += `<br>${conflict.slackLinks.map(link => `<a href="${link}" target="_blank">Slack Thread</a>`).join("<br>")}`
 		}
 
