Fix getPermalink for the dummy slack server by putting it correctly in the GET instead of POST section

diff --git a/Engine/Extras/RoboMerge/v3/src/common/dummyslackserver.ts b/Engine/Extras/RoboMerge/v3/src/common/dummyslackserver.ts
--- a/Engine/Extras/RoboMerge/v3/src/common/dummyslackserver.ts
+++ b/Engine/Extras/RoboMerge/v3/src/common/dummyslackserver.ts
@@ -44,11 +44,7 @@ export class DummySlackApp implements AppInterface {
 			return { dataObj }
 		}
 
-		if (command === "chat.getPermalink") {
-			// We don't need the permalink for any functional tests
-			return { ok: true }
-		}
-		else if (command === "chat.postMessage") {
+		if (command === "chat.postMessage") {
 			if (!data) {
 				throw new Error('Nothing to post!')
 			}
@@ -123,7 +119,11 @@ export class DummySlackApp implements AppInterface {
 
 	@Handler('GET', '/api/*')
 	get(command: string) {
-		if (command.startsWith("users.lookupByEmail"))
+		if (command === "chat.getPermalink") {
+			// We don't need the permalink for any functional tests
+			return { ok: true }
+		}
+		else if (command.startsWith("users.lookupByEmail"))
 		{
 			return { ok: true, user: { id: `${this.req.url.searchParams.get('email')}` } }
 		}
