Switch getChange to use the actual change -o command rather than changes which is known to return incorrect information from swarm periodically
Clean up Change interface

diff --git a/Engine/Extras/RoboMerge/v3/functional_tests/framework/src/test-perforce.ts b/Engine/Extras/RoboMerge/v3/functional_tests/framework/src/test-perforce.ts
--- a/Engine/Extras/RoboMerge/v3/functional_tests/framework/src/test-perforce.ts
+++ b/Engine/Extras/RoboMerge/v3/functional_tests/framework/src/test-perforce.ts
@@ -88,7 +88,7 @@ export class Perforce {
 	async changes(client: string, stream: string, limit: number, pending?: boolean) {
 		const ztagResult = await this.exec(['-ztag', '-c', client, 'changes', '-l', '-m', limit.toString(), '-s', pending ? 'pending' : 'submitted', stream + '/...'])
 
-		const parseResult = ztag.parseZtagOutput(ztagResult, 0, {expected: {change: 'integer', desc: 'string'}, optional: {oldChange: 'integer'}})
+		const parseResult = ztag.parseZtagOutput(ztagResult, 0, {expected: {change: 'integer', desc: 'string'}})
 		const result: Change[] = []
 		for (const entry of parseResult) {
 			result.push({change: entry.change as number, description: entry.desc as string})
diff --git a/Engine/Extras/RoboMerge/v3/src/common/perforce.ts b/Engine/Extras/RoboMerge/v3/src/common/perforce.ts
--- a/Engine/Extras/RoboMerge/v3/src/common/perforce.ts
+++ b/Engine/Extras/RoboMerge/v3/src/common/perforce.ts
@@ -226,13 +226,8 @@ export interface Change {
 	change: number;
 	client: string;
 	user: string;
-	// path?: string;
 	desc: string;
-	// status?: string;
 	shelved?: number;
-
-	// also from Perforce, but maybe not useful
-	changeType?: string;
 	time?: number;
 
 	// hacked in
@@ -292,6 +287,7 @@ interface ExecOpts {
 }
 
 interface ExecZtagOpts extends ExecOpts {
+	format?: string; // if specified will be passed as the -F argument and the results returned without ztag parsing
 	multiline?: boolean;
 	resolve?: boolean; // hacky solution to clear certain problematic lines out of resolve ztags
 	reduce?: boolean; // collapse the multiple entries to a single object, useful for problem parses that end up in multiple entries despite being a single result
@@ -520,20 +516,11 @@ export class PerforceContext {
 		return this._execP4Ztag(null, args, { multiline: true, edgeServerAddress });
 	}
 
-	/** get a single change in the format of changes() */
-	async getChange(path_in: string, changenum: number, status?: ChangelistStatus) {
-		const list = await this.changes(`${path_in}@${changenum},${changenum}`, -1, 1, status, false) as Change[]
-		if (list.length <= 0) {
-			throw new Error(`Could not find changelist ${changenum} in ${path_in}`);
-		}
-		if (list.length > 1 || list[0].change !== changenum) {
-			// log for now
-			const e = new Error();
-			this.logger.error(`${e.stack}\np4.getChange unexpected result for ${changenum}` +
-				list.map(change => `\n    ${change.change}: user ${change.user}, workspace ${change.client}`).join('')
-			)
-		}
-		return list[0]
+	/** get a single change and return it in the format of changes() */
+	async getChange(changenum: number) {
+		const result = await this._execP4Ztag(null, ['change', '-o', changenum.toString()], 
+						{format: '{"change":%Change%,"client":"%Client%","user":"%User%","status":"%Status%","desc":"%Description%"}'})
+		return JSON.parse(result.trimEnd().replaceAll("\n","\\n"));
 	}
 
 	/**
@@ -1765,7 +1752,12 @@ export class PerforceContext {
 
 	static async _execP4Ztag(logger: ContextualLogger, roboWorkspace: RoboWorkspace, args: string[], opts?: ExecZtagOpts) {
 		const workspace = coercePerforceWorkspace(roboWorkspace);
-		return parseZTag(await PerforceContext._execP4(logger, workspace, ['-ztag', ...args], opts), opts);
+		if (opts && opts.format) {
+			return PerforceContext._execP4(logger, workspace, ['-ztag', '-F', opts.format, ...args], opts)
+		}
+		else {
+			return parseZTag(await PerforceContext._execP4(logger, workspace, ['-ztag', ...args], opts), opts);
+		}
 	}
 
 	private async _execP4Ztag(roboWorkspace: RoboWorkspace, args: string[], opts?: ExecZtagOpts) {
diff --git a/Engine/Extras/RoboMerge/v3/src/new/graph.ts b/Engine/Extras/RoboMerge/v3/src/new/graph.ts
--- a/Engine/Extras/RoboMerge/v3/src/new/graph.ts
+++ b/Engine/Extras/RoboMerge/v3/src/new/graph.ts
@@ -834,7 +834,7 @@ export class Trace {
 
 		let change
 		try {
-			change = await this.p4.getChange(sourceNode.stream + '/...', cl)
+			change = await this.p4.getChange(cl)
 		}
 		catch (err) {
 			return 'CHANGELIST_NOT_FOUND'
diff --git a/Engine/Extras/RoboMerge/v3/src/robo/badges.ts b/Engine/Extras/RoboMerge/v3/src/robo/badges.ts
--- a/Engine/Extras/RoboMerge/v3/src/robo/badges.ts
+++ b/Engine/Extras/RoboMerge/v3/src/robo/badges.ts
@@ -193,7 +193,7 @@ async function markViaBranches(viaClStrings: string[], badgeFunc: BadgeFunc, bra
 
 		let result
 		try {
-			result = await p4.getChange(`//${branchGraph.config.defaultStreamDepot}/...`, cl)
+			result = await p4.getChange(cl)
 		}
 		catch (err) {
 			// this can fail sometimes with cross depot integrations
diff --git a/Engine/Extras/RoboMerge/v3/src/robo/nodebot.ts b/Engine/Extras/RoboMerge/v3/src/robo/nodebot.ts
--- a/Engine/Extras/RoboMerge/v3/src/robo/nodebot.ts
+++ b/Engine/Extras/RoboMerge/v3/src/robo/nodebot.ts
@@ -699,7 +699,7 @@ export class NodeBot extends PerforceStatefulBot implements NodeBotInterface {
 
 			let p4ChangeResult
 			try {
-				p4ChangeResult = await this.p4.getChange( targetBranch.rootPath, changelist, 'shelved')
+				p4ChangeResult = await this.p4.getChange(changelist)
 			}
 			catch (err) {
 				return { 
diff --git a/Engine/Extras/RoboMerge/v3/src/robo/perforce-stateful-bot.ts b/Engine/Extras/RoboMerge/v3/src/robo/perforce-stateful-bot.ts
--- a/Engine/Extras/RoboMerge/v3/src/robo/perforce-stateful-bot.ts
+++ b/Engine/Extras/RoboMerge/v3/src/robo/perforce-stateful-bot.ts
@@ -1,6 +1,6 @@
 // Copyright Epic Games, Inc. All Rights Reserved.
 import { ContextualLogger, NpmLogLevel } from "../common/logger";
-import { Change, ChangelistStatus, PerforceContext } from "../common/perforce";
+import { Change, PerforceContext } from "../common/perforce";
 import { BotIPC, ReconsiderArgs } from "./bot-interfaces";
 import { BlockagePauseInfo, BlockagePauseInfoMinimal } from "./status-types";
 import { Branch, OperationResult } from "./branch-interfaces";
@@ -154,9 +154,9 @@ export abstract class PerforceStatefulBot implements BotIPC {
 	abstract forceSetLastClWithContext(value: number, culprit: string, reason: string, unblock: boolean): number;
 	abstract setGateCl(value: number, culprit: string, reason: string): Promise<number | null>;
 
-	async _getChange(changeCl: number, path?: string, status?: ChangelistStatus) : Promise<Change | string> {
+	async _getChange(changeCl: number) : Promise<Change | string> {
 		try {
-			return await this.p4.getChange(path || this.branch.rootPath, changeCl, status)
+			return await this.p4.getChange(changeCl)
 		}
 		catch (err) {
 			return err.toString()
