Fix JSON parsing failing if a description had quotes in it

diff --git a/Engine/Extras/RoboMerge/v3/src/common/perforce.ts b/Engine/Extras/RoboMerge/v3/src/common/perforce.ts
--- a/Engine/Extras/RoboMerge/v3/src/common/perforce.ts
+++ b/Engine/Extras/RoboMerge/v3/src/common/perforce.ts
@@ -518,9 +518,16 @@ export class PerforceContext {
 
 	/** get a single change and return it in the format of changes() */
 	async getChange(changenum: number) {
-		const result = await this._execP4Ztag(null, ['change', '-o', changenum.toString()], 
+		let result = await this._execP4Ztag(null, ['change', '-o', changenum.toString()], 
 						{format: '{"change":%Change%,"client":"%Client%","user":"%User%","status":"%Status%","desc":"%Description%"}'})
-		return JSON.parse(result.trimEnd().replaceAll("\n","\\n"));
+		result = result.trimEnd().replaceAll("\n","\\n")
+		
+		// Have to escape any quotes in the description otherwise the JSON parsing fails
+		const descIndex = result.find('"desc":')
+		const descEndIndex = result.lastIndexOf('"')	
+		result = result.slice(0,descIndex+8) + result.slice(descIndex+8,descEndIndex).replaceAll('"','\\"') + '"}'	
+		
+		return JSON.parse(result);
 	}
 
 	/**
