[NO NON-PII COMMIT MESSAGE]

diff --git a/Engine/Extras/RoboMerge/v3/functional_tests/framework/src/framework.ts b/Engine/Extras/RoboMerge/v3/functional_tests/framework/src/framework.ts
--- a/Engine/Extras/RoboMerge/v3/functional_tests/framework/src/framework.ts
+++ b/Engine/Extras/RoboMerge/v3/functional_tests/framework/src/framework.ts
@@ -351,6 +351,7 @@ export abstract class FunctionalTest {
 	}
 
 	error(...err: string[]) {
+		console.log(`ERROR: ${err}`)
 		this.log(err, 'error')
 	}
 
