[NO NON-PII COMMIT MESSAGE]

diff --git a/Engine/Extras/RoboMerge/v3/functional_tests/framework/src/framework.ts b/Engine/Extras/RoboMerge/v3/functional_tests/framework/src/framework.ts
--- a/Engine/Extras/RoboMerge/v3/functional_tests/framework/src/framework.ts
+++ b/Engine/Extras/RoboMerge/v3/functional_tests/framework/src/framework.ts
@@ -351,7 +351,6 @@ export abstract class FunctionalTest {
 	}
 
 	error(...err: string[]) {
-		console.log(`ERROR: ${err}`)
 		this.log(err, 'error')
 	}
 
@@ -613,10 +612,13 @@ export abstract class FunctionalTest {
 		let verifyResult: any
 
 		try {
-			const post = bent('POST', 'json', 200, 400)
+			const post = bent('POST', 'json', 200, 400, 500)
 			verifyResult = await post(url)
 
-			// console.dir(verifyResult)
+			if (verifyResult.statusCode == 500) {
+				console.log(verifyResult)
+				throw new Error(`Verify returned 500 (${verifyResult})`)
+			}
 		}
 		catch (err) {
 			this.error(err)
