Fix typescript compile error

diff --git a/Engine/Extras/RoboMerge/v3/src/common/perforce.ts b/Engine/Extras/RoboMerge/v3/src/common/perforce.ts
--- a/Engine/Extras/RoboMerge/v3/src/common/perforce.ts
+++ b/Engine/Extras/RoboMerge/v3/src/common/perforce.ts
@@ -523,7 +523,7 @@ export class PerforceContext {
 		result = result.trimEnd().replaceAll("\n","\\n")
 		
 		// Have to escape any quotes in the description otherwise the JSON parsing fails
-		const descIndex = result.find('"desc":')
+		const descIndex = result.search('"desc":')
 		const descEndIndex = result.lastIndexOf('"')	
 		result = result.slice(0,descIndex+8) + result.slice(descIndex+8,descEndIndex).replaceAll('"','\\"') + '"}'	
 		
