Show All Bots and Track Changes for non-FTEs

diff --git a/Engine/Extras/RoboMerge/v3/public/css/main.css b/Engine/Extras/RoboMerge/v3/public/css/main.css
--- a/Engine/Extras/RoboMerge/v3/public/css/main.css
+++ b/Engine/Extras/RoboMerge/v3/public/css/main.css
@@ -62,7 +62,7 @@ a:focus, a:hover {
 	margin: .1em;
 }
 
-#fixed-footer div#fteButtons {
+#fixed-footer div#buttons {
 	align-self: flex-start;
 }
 
diff --git a/Engine/Extras/RoboMerge/v3/public/js/boilerplate.js b/Engine/Extras/RoboMerge/v3/public/js/boilerplate.js
--- a/Engine/Extras/RoboMerge/v3/public/js/boilerplate.js
+++ b/Engine/Extras/RoboMerge/v3/public/js/boilerplate.js
@@ -300,7 +300,10 @@ function handleUserPermissions(user) {
 
 	// Fulltime employees should see the log buttons.
 	if (isFTE) {
-		$('#fteButtons').show();
+		$('#logButton').show();
+		$('#lastCrashButton').show();
+		$('#p4TasksButton').show();
+		$('#branchesButton').show();
 	}
 }
 
@@ -451,39 +454,42 @@ function generateRobomergeHeader(createSignedInUserDiv = true) {
 function generateRobomergeFooter() {
 	let fixedFooterContents = $('<div class="button-bar">')
 
-	let fteButtonDiv = $('<div id="fteButtons">')
-	fteButtonDiv.hide()
-	fixedFooterContents.append(fteButtonDiv)
+	let buttonDiv = $('<div id="buttons">')
+	fixedFooterContents.append(buttonDiv)
 
 	let logButton = $('<button id="logButton">')
 	logButton.addClass("btn btn-sm btn-outline-dark")
 	logButton.click(function() { window.open('/api/logs', '_blank') })
 	logButton.text("Logs")
-	fteButtonDiv.append(logButton)
+	logButton.hide()
+	buttonDiv.append(logButton)
 
 	let lastCrashButton = $('<button id="lastCrashButton">')
 	lastCrashButton.addClass("btn btn-sm btn-outline-dark")
 	lastCrashButton.click(function() { window.open('/api/last_crash', '_blank') })
 	lastCrashButton.text("Last Crash")
-	fteButtonDiv.append(lastCrashButton)
+	lastCrashButton.hide()
+	buttonDiv.append(lastCrashButton)
 
 	let p4TasksButton = $('<button id="p4TasksButton">')
 	p4TasksButton.addClass("btn btn-sm btn-outline-dark")
 	p4TasksButton.click(function() { window.open('/api/p4tasks', '_blank') })
 	p4TasksButton.text("P4 Tasks")
-	fteButtonDiv.append(p4TasksButton)
-
-	let p4AllBotsButton = $('<button id="p4AllBotsButton">')
-	p4AllBotsButton.addClass("btn btn-sm btn-outline-dark")
-	p4AllBotsButton.click(function() { window.open('/allbots', '_blank') })
-	p4AllBotsButton.text("All bots graph")
-	fteButtonDiv.append(p4AllBotsButton)
+	p4TasksButton.hide()
+	buttonDiv.append(p4TasksButton)
 
 	let branchesButton = $('<button id="branchesButton">')
 	branchesButton.addClass("btn btn-sm btn-outline-dark")
 	branchesButton.click(function() { window.open('/api/branches', '_blank') })
 	branchesButton.text("Branch Data")
-	fteButtonDiv.append(branchesButton)
+	branchesButton.hide()
+	buttonDiv.append(branchesButton)
+
+	let p4AllBotsButton = $('<button id="p4AllBotsButton">')
+	p4AllBotsButton.addClass("btn btn-sm btn-outline-dark")
+	p4AllBotsButton.click(function() { window.open('/allbots', '_blank') })
+	p4AllBotsButton.text("All bots graph")
+	buttonDiv.append(p4AllBotsButton)
 
 	let trackChangeButton = $('<button id="trackChangeButton">')
 	trackChangeButton.addClass("btn btn-sm btn-outline-dark")
@@ -501,7 +507,7 @@ function generateRobomergeFooter() {
 		}
 	})
 	trackChangeButton.text("Track Change")
-	fteButtonDiv.append(trackChangeButton)
+	buttonDiv.append(trackChangeButton)
 
 	let currentlyRunningDiv = $('<div id="currentlyRunning">')
 	fixedFooterContents.append(currentlyRunningDiv)
