Fix error parsing json result from p4 change -F for changelists that have a tab in the changelist description

diff --git a/Engine/Extras/RoboMerge/v3/src/common/perforce.ts b/Engine/Extras/RoboMerge/v3/src/common/perforce.ts
--- a/Engine/Extras/RoboMerge/v3/src/common/perforce.ts
+++ b/Engine/Extras/RoboMerge/v3/src/common/perforce.ts
@@ -520,7 +520,7 @@ export class PerforceContext {
 	async getChange(changenum: number) {
 		let result = await this._execP4Ztag(null, ['change', '-o', changenum.toString()], 
 						{format: '{"change":%Change%,"client":"%Client%","user":"%User%","status":"%Status%","desc":"%Description%"}'})
-		result = result.trimEnd().replaceAll("\n","\\n")
+		result = result.trimEnd().replaceAll("\n","\\n").replaceAll("\t","\\t")
 		
 		// Have to escape any quotes in the description otherwise the JSON parsing fails
 		const descIndex = result.search('"desc":')
