Shelves created via reconsider will now include #ROBOMERGE-CONFLICT from-shelf consistent with shelves made via Create Shelf when a change is blocked

diff --git a/Engine/Extras/RoboMerge/v3/src/robo/edgebot.ts b/Engine/Extras/RoboMerge/v3/src/robo/edgebot.ts
--- a/Engine/Extras/RoboMerge/v3/src/robo/edgebot.ts
+++ b/Engine/Extras/RoboMerge/v3/src/robo/edgebot.ts
@@ -291,6 +291,7 @@ class EdgeBotImpl extends PerforceStatefulBot {
 		if (pending.change.userRequest) {
 			const owner = getIntegrationOwner(pending) || pending.change.author
 			const shelfMsg = `${owner}, please merge this change by hand.\nMore info at ${this.sourceNode.getBotUrl()}\n\n` + failure.description
+			pending.change.additionalDescriptionText = `#ROBOMERGE-CONFLICT from-shelf\n`
 			await this.shelveChangelist(pending, false, shelfMsg)
 			this.edgeBotLogger.info(`${logMessage}. Shelved CL ${pending.newCl} for ${owner} to resolve manually (from reconsider).`)
 			return
