Warn a user when the reconsider or #manual shelf was created in a workspace that was picked by partial match and does not (necessarily) match the intended stream
Minor optimization when picking best workspace

diff --git a/Engine/Extras/RoboMerge/v3/src/common/p4util.ts b/Engine/Extras/RoboMerge/v3/src/common/p4util.ts
--- a/Engine/Extras/RoboMerge/v3/src/common/p4util.ts
+++ b/Engine/Extras/RoboMerge/v3/src/common/p4util.ts
@@ -255,14 +255,17 @@ export async function chooseBestWorkspaceForUser(p4: PerforceContext, user: stri
 				let stream = def.Stream
 				if (stream) {
 					const matchlen = matchPrefix(stream.toLowerCase(), branch_stream)
-					if (matchlen > target_match) {
+					if (matchlen == branch_stream.length) {
+						return def
+					}
+					else if (matchlen > target_match) {
 						target_match = matchlen
 						targetWorkspace = def
 					}
 				}
 			}
 		}
-		return targetWorkspace.client
+		return targetWorkspace
 	}
 
 	return undefined
diff --git a/Engine/Extras/RoboMerge/v3/src/robo/branch-interfaces.ts b/Engine/Extras/RoboMerge/v3/src/robo/branch-interfaces.ts
--- a/Engine/Extras/RoboMerge/v3/src/robo/branch-interfaces.ts
+++ b/Engine/Extras/RoboMerge/v3/src/robo/branch-interfaces.ts
@@ -143,6 +143,7 @@ export interface TargetInfo {
 	author: string
 
 	targetWorkspaceForShelf?: string // Filled in during the reconsider in case of a createShelf nodeop request
+	targetWorkspaceIsPartialMatch?: boolean // Filled in during the reconsider with whether the workspace match was partial or not
 	sendNoShelfNotification: boolean // Used for internal use shelves, such as stomp changes
 
 	forceStompChanges: boolean
diff --git a/Engine/Extras/RoboMerge/v3/src/robo/edgebot.ts b/Engine/Extras/RoboMerge/v3/src/robo/edgebot.ts
--- a/Engine/Extras/RoboMerge/v3/src/robo/edgebot.ts
+++ b/Engine/Extras/RoboMerge/v3/src/robo/edgebot.ts
@@ -849,9 +849,24 @@ class EdgeBotImpl extends PerforceStatefulBot {
 		// Find a suitable workspace from one of the owner's workspaces
 		else if (!forApproval) {
 			// use p4.find_workspaces to find a workspace (owned by the user) for this change if this is a stream branch
-			targetWorkspace = await p4util.chooseBestWorkspaceForUser(this.p4, owner, branch_stream)
+			const targetWorkspaceDef = await p4util.chooseBestWorkspaceForUser(this.p4, owner, branch_stream)
+			if (targetWorkspaceDef) {
+				targetWorkspace = targetWorkspaceDef.client
+				if (targetWorkspaceDef.Stream) {
+					pending.change.targetWorkspaceIsPartialMatch = targetWorkspaceDef.Stream.toLowerCase() == branch_stream
+				}
+				if (pending.change.targetWorkspaceIsPartialMatch) {
+					this.edgeBotLogger.info(`Chose workspace ${targetWorkspace} (${targetWorkspaceDef.Stream}) as a partial match for ${branch_stream}`)
+				}
+				else {
+					this.edgeBotLogger.info(`Chose workspace ${targetWorkspace}`)
+				}
+			}
+			else {
+				targetWorkspace = undefined
+				this.edgeBotLogger.error(`Unable to find workspace for ${branch_stream}`)
+			}
 			pending.change.targetWorkspaceForShelf = targetWorkspace
-			this.edgeBotLogger.info(`Chose workspace ${targetWorkspace}`)
 		}
 
 		// log if we couldn't find a workspace
diff --git a/Engine/Extras/RoboMerge/v3/src/robo/nodebot.ts b/Engine/Extras/RoboMerge/v3/src/robo/nodebot.ts
--- a/Engine/Extras/RoboMerge/v3/src/robo/nodebot.ts
+++ b/Engine/Extras/RoboMerge/v3/src/robo/nodebot.ts
@@ -1593,9 +1593,21 @@ export class NodeBot extends PerforceStatefulBot implements NodeBotInterface {
 		let isManualChange = numTargets > 0 ? result.info.targets![0].flags.has('manual') : false;
 		if (change.forceCreateAShelf || (change.isUserRequest && !change.forceStompChanges) || isManualChange) {
 			if (!optWorkspaceOverride && numTargets == 1) {
-				result.info.targetWorkspaceForShelf = await p4util.chooseBestWorkspaceForUser(this.p4, result.info.owner||result.info.author, result.info.targets![0].branch.stream)
+				const targetStream = result.info.targets![0].branch.stream ? result.info.targets![0].branch.stream.toLowerCase() : undefined
+				const targetWorkspaceDef = await p4util.chooseBestWorkspaceForUser(this.p4, result.info.owner||result.info.author, targetStream)
+				if (targetWorkspaceDef) {
+					result.info.targetWorkspaceForShelf = targetWorkspaceDef.client
+					if (targetWorkspaceDef.Stream) {
+						result.info.targetWorkspaceIsPartialMatch = targetWorkspaceDef.Stream.toLowerCase() == targetStream
+					}
+				}
 				optWorkspaceOverride = result.info.targetWorkspaceForShelf
-				this.nodeBotLogger.info(`Chose workspace ${optWorkspaceOverride}`)
+				if (result.info.targetWorkspaceIsPartialMatch) {
+					this.nodeBotLogger.info(`Chose workspace ${optWorkspaceOverride} (${targetWorkspaceDef!.Stream}) as a partial match for ${targetStream}`)
+				}
+				else {
+					this.nodeBotLogger.info(`Chose workspace ${optWorkspaceOverride}`)
+				}
 			}
 			if (optWorkspaceOverride) {
 				// see if we need to talk to an edge server to create the shelf
@@ -2128,6 +2140,12 @@ export class NodeBot extends PerforceStatefulBot implements NodeBotInterface {
 
 			if (postIntegrate && pending.change.targetWorkspaceForShelf) {
 				message += `\n\nShelved ${pending.newCl} in workspace ${pending.change.targetWorkspaceForShelf}`
+
+				if (pending.change.targetWorkspaceIsPartialMatch) {
+					message += `\n\n*NOTE: The target workspace was not an exact match for the shelved files. `
+					message += "You may need to remap the workspace to the target stream, move the files to "
+					message += "another workspace, or create a workspace for these files and reconsider again.*"
+				}
 			}
 
 			if (this.slackMessages) {
