Add messaging clarifying that creating a shelf does unblock robomerge

diff --git a/Engine/Extras/RoboMerge/v3/public/createshelf.html b/Engine/Extras/RoboMerge/v3/public/createshelf.html
--- a/Engine/Extras/RoboMerge/v3/public/createshelf.html
+++ b/Engine/Extras/RoboMerge/v3/public/createshelf.html
@@ -29,7 +29,8 @@
     <div id="result"></div>
     <div id="nextSteps" class="initiallyHidden centered-text">
         <h3>What next?</h3>
-        <p>Robomerge will run the merge for the changelist again and create any shelved conflicts in the workspace you chose. You should receive a notification when this process is complete.</p>
+        <p>Robomerge will run the merge for the changelist again and create any shelved conflicts in the workspace you chose. You should receive a notification when this process is complete. 
+        <b>This does not skip or unblock robomerge. It is still time sensitive that you resolve the conflict and submit, or skip the conflict.</b></p>
         <p>If you do not receive a notification, please <strong><a href="/contact">Contact Us</a></strong> for Robomerge help.</p>
     </div>
     <div id="createStreamWorkspaceTutorial" class="initiallyHidden centered-text">
diff --git a/Engine/Extras/RoboMerge/v3/src/robo/nodebot.ts b/Engine/Extras/RoboMerge/v3/src/robo/nodebot.ts
--- a/Engine/Extras/RoboMerge/v3/src/robo/nodebot.ts
+++ b/Engine/Extras/RoboMerge/v3/src/robo/nodebot.ts
@@ -1459,13 +1459,18 @@ export class NodeBot extends PerforceStatefulBot implements NodeBotInterface {
 			return
 		}
 
-		this.nodeBotLogger.info(`Sending reconsider shelf creation (shelf CL ${pendingChange.newCl}) notification to ${pendingChange.change.owner} for source CL ${pendingChange.change.source_cl}`)
+		this.nodeBotLogger.info(`Sending shelf creation (shelf CL ${pendingChange.newCl}) notification to ${pendingChange.change.owner} for source CL ${pendingChange.change.source_cl}`)
+
+		let showStillBlockedMessage = pendingChange.change.forceCreateAShelf && !pendingChange.action.flags.has('manual')
 
 		if (this.slackMessages) {
 			let dm: SlackMessage = {
 				text: `Robomerge has created shelf CL ${pendingChange.newCl} in workspace ${pendingChange.change.targetWorkspaceForShelf} ` +
 				`for merging source CL ${pendingChange.change.source_cl} (${this.fullName}).\n` +
-				`\`\`\`${pendingChange.change.description}\`\`\``,
+				`\`\`\`${pendingChange.change.description}\`\`\`` +
+				(showStillBlockedMessage	?
+					'\n\n*Creating a shelf does not skip or unblock robomerge. It is still time sensitive that you resolve the conflict and submit, or skip the conflict.*'
+					: ''),
 				channel: '',
 				mrkdwn: true
 			}
@@ -1479,7 +1484,10 @@ export class NodeBot extends PerforceStatefulBot implements NodeBotInterface {
 				new Recipients(owner),
 				`Robomerge created Shelf CL ${pendingChange.newCl} in ${pendingChange.change.targetWorkspaceForShelf}`,
 				`Robomerge has created shelf CL ${pendingChange.newCl} in workspace ${pendingChange.change.targetWorkspaceForShelf} ` +
-					`for merging source CL ${pendingChange.change.source_cl} (${this.fullName}).`,
+					`for merging source CL ${pendingChange.change.source_cl} (${this.fullName}).` +
+				showStillBlockedMessage	?
+					'\n\nCreating a shelf does not skip or unblock robomerge. It is still time sensitive that you resolve the conflict and submit, or skip the conflict.'
+					: '',
 				`${pendingChange.change.source_cl}:\n${pendingChange.change.description}`)
 		}
 	}
