Don't crash if unable to retrieve address for server
Properly communicate error when failing to unlock

diff --git a/Engine/Extras/RoboMerge/v3/public/js/unlock.js b/Engine/Extras/RoboMerge/v3/public/js/unlock.js
--- a/Engine/Extras/RoboMerge/v3/public/js/unlock.js
+++ b/Engine/Extras/RoboMerge/v3/public/js/unlock.js
@@ -172,7 +172,7 @@ function unlockVerify() {
     }).fail(function(jqXHR, textStatus, errMsg) {
         try {
             transitionDisplayUnlockResults()
-            unlockFailure(`${textStatus}: ${errMsg}`)
+            unlockFailure(jqXHR.responseText)
         }
         catch (err) {
             unlockFailure(`Error encountered displaying unlock files failure: ${err}`)
diff --git a/Engine/Extras/RoboMerge/v3/src/robo/nodebot.ts b/Engine/Extras/RoboMerge/v3/src/robo/nodebot.ts
--- a/Engine/Extras/RoboMerge/v3/src/robo/nodebot.ts
+++ b/Engine/Extras/RoboMerge/v3/src/robo/nodebot.ts
@@ -1077,7 +1077,20 @@ export class NodeBot extends PerforceStatefulBot implements NodeBotInterface {
 			filesByUser.set(lockedFile.user, [...(filesByUser.get(lockedFile.user) || []), lockedFile.depotPath])
 		}
 
-		filesByClient.forEach((files,client) => this.p4.revertFiles(files,client))
+		let results = await Promise.all(Array.from(filesByClient.entries()).map(async ([client,files]) => 
+			{ 
+				try { 
+					await this.p4.revertFiles(files,client)
+				} catch (reason) {
+					return { success: false, message: reason }
+				}
+				return { success: true }
+			}))
+		results = results.filter(result => !result.success)
+		
+		if (results.length > 0) {
+			return { success: false, message: results.map(result => result.message).join('\n') }
+		}
 
 		if (this.slackMessages) {
 			
