Fix determination of hasAutomerge failing in certain cases

diff --git a/Engine/Extras/RoboMerge/v3/src/robo/ipc.ts b/Engine/Extras/RoboMerge/v3/src/robo/ipc.ts
--- a/Engine/Extras/RoboMerge/v3/src/robo/ipc.ts
+++ b/Engine/Extras/RoboMerge/v3/src/robo/ipc.ts
@@ -291,8 +291,13 @@ export class IPC {
 							}
 						}
 					}
-					if (!hasAutomergeTarget && edge.flags.has('automatic') && edge.source.stream == changeToConsider.node) {
-						hasAutomergeTarget = true
+					if (!hasAutomergeTarget && edge.flags.has('automatic')) {
+						if (typeof changeToConsider.node === 'string') {
+							hasAutomergeTarget = edge.source.stream == changeToConsider.node
+						}
+						else {
+							hasAutomergeTarget = edge.source.stream == changeToConsider.node.stream
+						}
 					}
 					if (includeInResults && hasAutomergeTarget) {
 						break
