Hide allbots and trackchanges when logged out

diff --git a/Engine/Extras/RoboMerge/v3/public/js/boilerplate.js b/Engine/Extras/RoboMerge/v3/public/js/boilerplate.js
--- a/Engine/Extras/RoboMerge/v3/public/js/boilerplate.js
+++ b/Engine/Extras/RoboMerge/v3/public/js/boilerplate.js
@@ -298,7 +298,12 @@ function handleUserPermissions(user) {
 		}
 	}
 
-	// Fulltime employees should see the log buttons.
+	// Everyone sees when logged in
+	$('#p4AllBotsButton').show();
+	$('#trackChangeButton').show();
+
+
+	// Fulltime employees should see a wider set of buttons.
 	if (isFTE) {
 		$('#logButton').show();
 		$('#lastCrashButton').show();
@@ -489,6 +494,7 @@ function generateRobomergeFooter() {
 	p4AllBotsButton.addClass("btn btn-sm btn-outline-dark")
 	p4AllBotsButton.click(function() { window.open('/allbots', '_blank') })
 	p4AllBotsButton.text("All bots graph")
+	p4AllBotsButton.hide()
 	buttonDiv.append(p4AllBotsButton)
 
 	let trackChangeButton = $('<button id="trackChangeButton">')
@@ -507,6 +513,7 @@ function generateRobomergeFooter() {
 		}
 	})
 	trackChangeButton.text("Track Change")
+	trackChangeButton.hide()
 	buttonDiv.append(trackChangeButton)
 
 	let currentlyRunningDiv = $('<div id="currentlyRunning">')
