Add mechanism to specify okta token (found in browser cookies) for servers using okta login

diff --git a/Engine/Extras/RoboMerge/v3/repl/roborepl.py b/Engine/Extras/RoboMerge/v3/repl/roborepl.py
--- a/Engine/Extras/RoboMerge/v3/repl/roborepl.py
+++ b/Engine/Extras/RoboMerge/v3/repl/roborepl.py
@@ -116,38 +116,40 @@ def fetch_branch_data():
 	bots = context['_bots']
 	bots.clear()
 
-	branches = []
-	edge_statuses = []
-	for branch_json in json.loads(response.text)['branches']:
-		bot_name = branch_json['bot']
-		bot = bots.setdefault(bot_name.lower(), Bot(bot_name))
-
-		branch = Branch(bot, branch_json)
-		branches.append(branch)
-
-		edges = branch_json.get('edges')
-		if edges:
-			edge_statuses += [(bot, branch.name, target, status) for target, status in edges.items()]
-
-		node = Node(bot, branch.name)
-
-		bot.nodes[branch.name.lower()] = node
-		for alias in branch.aliases:
-			bot.nodes[alias.lower()] = node
-
-	for branch in branches:
-		for target in branch.flows_to:
-			source_node = branch.bot.nodes[branch.name.lower()]
-			target_node = branch.bot.nodes[target.lower()]
-			branch.bot.edges[(source_node, target_node)] \
-				= Edge(branch.bot, source_node, target_node,
-						'auto' if target in branch.force_flow else 'normal'
-								# not quite right if alias used for force list
-				)
-
-	for bot, branch_name, target, status_json in edge_statuses:
-		find_edge(bot.name, branch_name, target).update_status(status_json)
-
+	try:
+		branches = []
+		edge_statuses = []
+		for branch_json in json.loads(response.text)['branches']:
+			bot_name = branch_json['bot']
+			bot = bots.setdefault(bot_name.lower(), Bot(bot_name))
+
+			branch = Branch(bot, branch_json)
+			branches.append(branch)
+
+			edges = branch_json.get('edges')
+			if edges:
+				edge_statuses += [(bot, branch.name, target, status) for target, status in edges.items()]
+
+			node = Node(bot, branch.name)
+
+			bot.nodes[branch.name.lower()] = node
+			for alias in branch.aliases:
+				bot.nodes[alias.lower()] = node
+
+		for branch in branches:
+			for target in branch.flows_to:
+				source_node = branch.bot.nodes[branch.name.lower()]
+				target_node = branch.bot.nodes[target.lower()]
+				branch.bot.edges[(source_node, target_node)] \
+					= Edge(branch.bot, source_node, target_node,
+							'auto' if target in branch.force_flow else 'normal'
+									# not quite right if alias used for force list
+					)
+
+		for bot, branch_name, target, status_json in edge_statuses:
+			find_edge(bot.name, branch_name, target).update_status(status_json)
+	except:
+		return
 
 @export('login')
 def login(user, password=None):
@@ -166,6 +168,16 @@ def login(user, password=None):
 	else:
 		print(r.text)
 
+@export('setAuthToken')
+def setAuthToken(token):
+	context['_auth_token'] = token
+	fetch_branch_data()
+
+	if len(context['_bots']) > 0:
+		print("Token set successfully")
+	else:
+		print("Token did not allow for branch data to be fetched")
+
 # discoverability
 
 # query branches once for now - maybe provide way to refresh status
