Help and Contact should not be secure commands

diff --git a/Engine/Extras/RoboMerge/v3/src/robo/roboserver.ts b/Engine/Extras/RoboMerge/v3/src/robo/roboserver.ts
--- a/Engine/Extras/RoboMerge/v3/src/robo/roboserver.ts
+++ b/Engine/Extras/RoboMerge/v3/src/robo/roboserver.ts
@@ -328,10 +328,10 @@ class RoboWebApp implements AppInterface {
 		})
 	}
 
-	@SecureHandler('GET', '/help', {filetype: 'text/html'}) 
+	@Handler('GET', '/help', {filetype: 'text/html'}) 
 	renderHelp() { return RoboWebApp.renderMarkdownFile('README.md') }
 
-	@SecureHandler('GET', '/contact', {filetype: 'text/html'}) 
+	@Handler('GET', '/contact', {filetype: 'text/html'}) 
 	renderContactInfo() { return RoboWebApp.renderMarkdownFile('ContactInfo.md') }
 
 	@SecureHandler('GET', '/api/logs', {requiredTags: ['fte']})
