When navigating directly to a page when signed out and authentication method is okta, redirect correctly back to the page instead of the default page

diff --git a/Engine/Extras/RoboMerge/v3/public/allbots.html b/Engine/Extras/RoboMerge/v3/public/allbots.html
--- a/Engine/Extras/RoboMerge/v3/public/allbots.html
+++ b/Engine/Extras/RoboMerge/v3/public/allbots.html
@@ -26,11 +26,9 @@
 					...parseOptions(location.search)
 				})));
 			},
-		// error: function(xhr, error, status) {
-		// 	robomergeUser = undefined;
-		// 	setError(xhr, error);
-		// 	callback(null);
-		// }
+			error: function(xhr, error, status) {
+				window.location = window.origin + "/login?redirect=/allbots";
+			}
 		})
 	</script>
 </head>
diff --git a/Engine/Extras/RoboMerge/v3/public/js/login.js b/Engine/Extras/RoboMerge/v3/public/js/login.js
--- a/Engine/Extras/RoboMerge/v3/public/js/login.js
+++ b/Engine/Extras/RoboMerge/v3/public/js/login.js
@@ -65,7 +65,13 @@ $(async () => {
 					if (urlParams.has("redirect")) {
 						redirectString = decodeURIComponent(urlParams.get("redirect"))
 					}
-
+					else {
+						const redirectCookie = getCookie('redirect_to')
+						if (redirectCookie) {
+							redirectString = redirectCookie
+						}
+					}
+		
 					document.cookie = `auth=${token}; secure=true; path=/;`;
 					document.cookie = 'redirect_to=;';
 					window.location = window.origin + redirectString;
@@ -150,6 +156,11 @@ async function oktaSignIn() {
 	// Clear auth cookie on attempt to sign in
 	document.cookie = 'auth=;';
 
+	var urlParams = new URLSearchParams(window.location.search);
+	if (urlParams.has("redirect")) {
+		document.cookie = `redirect_to=${decodeURIComponent(urlParams.get("redirect"))}`
+	}
+
 	// Start full-page redirect to Okta
 	if (authClient) {
 		authClient.token.getWithRedirect();
