When setting the last considered CL correctly broadcast it the slack channels for that edge

diff --git a/Engine/Extras/RoboMerge/v3/src/robo/branch-interfaces.ts b/Engine/Extras/RoboMerge/v3/src/robo/branch-interfaces.ts
--- a/Engine/Extras/RoboMerge/v3/src/robo/branch-interfaces.ts
+++ b/Engine/Extras/RoboMerge/v3/src/robo/branch-interfaces.ts
@@ -214,6 +214,8 @@ export interface AlreadyIntegrated {
 
 export interface ForcedCl {
 	nodeOrEdgeName: string
+	sourceBranchUpperName: string
+	targetBranchUpperName: string
 	forcedCl: number
 	previousCl: number
 	culprit: string
diff --git a/Engine/Extras/RoboMerge/v3/src/robo/edgebot.ts b/Engine/Extras/RoboMerge/v3/src/robo/edgebot.ts
--- a/Engine/Extras/RoboMerge/v3/src/robo/edgebot.ts
+++ b/Engine/Extras/RoboMerge/v3/src/robo/edgebot.ts
@@ -972,7 +972,7 @@ class EdgeBotImpl extends PerforceStatefulBot {
 		}
 
 		// trigger events
-		this.sourceNode.onForcedLastCl(this.displayName, value, prevValue, culprit, reason)
+		this.sourceNode.onForcedLastCl(this.displayName, this.targetBranch.upperName, value, prevValue, culprit, reason)
 		
 		return prevValue
 	}
diff --git a/Engine/Extras/RoboMerge/v3/src/robo/nodebot.ts b/Engine/Extras/RoboMerge/v3/src/robo/nodebot.ts
--- a/Engine/Extras/RoboMerge/v3/src/robo/nodebot.ts
+++ b/Engine/Extras/RoboMerge/v3/src/robo/nodebot.ts
@@ -512,7 +512,7 @@ export class NodeBot extends PerforceStatefulBot implements NodeBotInterface {
 		const prevValue = this._forceSetLastCl_NoReset(value)
 
 		if (prevValue !== value) {
-			this.onForcedLastCl(this.displayName, value, prevValue, culprit, reason)
+			this.onForcedLastCl(this.displayName, this.branch.upperName, value, prevValue, culprit, reason)
 		}
 		return prevValue
 	}
@@ -523,8 +523,8 @@ export class NodeBot extends PerforceStatefulBot implements NodeBotInterface {
 		throw err
 	}
 
-	onForcedLastCl(nodeOrEdgeName: string, forcedCl: number, previousCl: number, culprit: string, reason: string) {
-		this.conflicts.onForcedLastCl({nodeOrEdgeName, forcedCl, previousCl, culprit, reason})
+	onForcedLastCl(nodeOrEdgeName: string, targetBranchUpperName: string, forcedCl: number, previousCl: number, culprit: string, reason: string) {
+		this.conflicts.onForcedLastCl({nodeOrEdgeName, sourceBranchUpperName: this.branch.upperName, targetBranchUpperName, forcedCl, previousCl, culprit, reason})
 	}
 
 	persistQueuedChanges() {
diff --git a/Engine/Extras/RoboMerge/v3/src/robo/notifications.ts b/Engine/Extras/RoboMerge/v3/src/robo/notifications.ts
--- a/Engine/Extras/RoboMerge/v3/src/robo/notifications.ts
+++ b/Engine/Extras/RoboMerge/v3/src/robo/notifications.ts
@@ -736,19 +736,34 @@ export class BotNotifications implements BotEventHandler {
 
 	onNonSkipLastClChange(details: ForcedCl) {
 		if (this.slackMessages) {
-			this.slackMessages.postNonConflictMessage({
-				title: details.nodeOrEdgeName + ' forced to new CL',
-				text: details.reason,
-				style: SlackMessageStyles.WARNING,
-				fields: [{
-					title: 'By', short: true, value: details.culprit
-				}, {
-					title: 'Changelists', short: true, value: `${makeClLink(details.previousCl)} -> ${makeClLink(details.forcedCl)}`
-				}],
-				title_link: this.externalRobomergeUrl + '#' + this.botname,
-				mrkdwn: true,
-				channel: this.slackChannel // Default to the configured channel
-			}) 
+			let channelsToPostTo = []		
+			const additionalChannelInfo = this.additionalBlockChannelIds.get(details.sourceBranchUpperName + '|' + details.targetBranchUpperName)
+			if (additionalChannelInfo) {
+				const [sideChannel, postOnlyToSideChannel] = additionalChannelInfo
+				if (!postOnlyToSideChannel) {
+					channelsToPostTo.push(this.slackChannel)
+				}
+				channelsToPostTo.push(sideChannel)
+			}
+			else {
+				channelsToPostTo.push(this.slackChannel)
+			}
+
+			for (let slackChannel of channelsToPostTo) {
+				this.slackMessages.postNonConflictMessage({
+					title: details.nodeOrEdgeName + ' forced to new CL',
+					text: details.reason,
+					style: SlackMessageStyles.WARNING,
+					fields: [{
+						title: 'By', short: true, value: details.culprit
+					}, {
+						title: 'Changelists', short: true, value: `${makeClLink(details.previousCl)} -> ${makeClLink(details.forcedCl)}`
+					}],
+					title_link: this.externalRobomergeUrl + '#' + this.botname,
+					mrkdwn: true,
+					channel: slackChannel // Default to the configured channel
+				}) 
+			}
 		}
 	}
 
