Fix spelling error

diff --git a/Engine/Extras/RoboMerge/v3/src/robo/edgebot.ts b/Engine/Extras/RoboMerge/v3/src/robo/edgebot.ts
--- a/Engine/Extras/RoboMerge/v3/src/robo/edgebot.ts
+++ b/Engine/Extras/RoboMerge/v3/src/robo/edgebot.ts
@@ -355,7 +355,7 @@ class EdgeBotImpl extends PerforceStatefulBot {
 			description += `#ROBOMERGE-OWNER: ${overriddenOwner}\n`
 		}
 
-		// keep track of author in a tag in case transfering onwership of the changelist fails
+		// keep track of author in a tag in case transfering ownership of the changelist fails
 		const authorTag = info.authorTag || info.author
 		if (authorTag !== 'robomerge') {
 			description += `#ROBOMERGE-AUTHOR: ${authorTag}\n`
