Tag the owner (if target stream has a resolver or the change has been reconsidered by another user) on approvals in addition to the author

diff --git a/Engine/Extras/RoboMerge/v3/src/robo/nodebot.ts b/Engine/Extras/RoboMerge/v3/src/robo/nodebot.ts
--- a/Engine/Extras/RoboMerge/v3/src/robo/nodebot.ts
+++ b/Engine/Extras/RoboMerge/v3/src/robo/nodebot.ts
@@ -2225,18 +2225,24 @@ export class NodeBot extends PerforceStatefulBot implements NodeBotInterface {
 
 	async reportApprovalRequired(approval: ApprovalOptions, pending: PendingChange) {
 		if (this.slackMessages) {
-			const userEmail = await this.p4.getEmail(pending.change.author)
-			const slackUser = userEmail ? await this.slackMessages.getSlackUser(userEmail) : null
-			const channelPing = slackUser ? `<@${slackUser}>` : `@${pending.change.author}`
+			const authorEmail = await this.p4.getEmail(pending.change.author)
+			const slackAuthor = authorEmail ? await this.slackMessages.getSlackUser(authorEmail) : null
+			const ownerEmail = pending.change.owner && pending.change.owner != pending.change.author ? await this.p4.getEmail(pending.change.owner) : null
+			const slackOwner = ownerEmail ?  await this.slackMessages.getSlackUser(ownerEmail) : null
 
 			const fields: SlackMessageField[] = [
 				{title: 'Change', short: true, value: makeClLink(pending.newCl)}
 			]
 
+			let message = slackAuthor ? `<@${slackAuthor}>` : `@${pending.change.author}` +
+						  "'s change " +
+						  (ownerEmail ? "owned by " + (slackOwner ? `<@${slackOwner}> ` : `@${pending.change.owner} `) : "") +
+						  `in ${pending.action.branch.name} needs to be approved.\n\n` +
+						  approval.description
+
 			const opts: SlackMessage = { 
 				title:'', 
-				text: `${channelPing}'s change in ${pending.action.branch.name} needs to be approved.\n\n` +
-						approval.description, 
+				text: message, 
 				style: SlackMessageStyles.DANGER, 
 				fields,
 				mrkdwn: true,
