An edge can override a nodes channel

diff --git a/Engine/Extras/RoboMerge/v3/src/robo/branchgraph.ts b/Engine/Extras/RoboMerge/v3/src/robo/branchgraph.ts
--- a/Engine/Extras/RoboMerge/v3/src/robo/branchgraph.ts
+++ b/Engine/Extras/RoboMerge/v3/src/robo/branchgraph.ts
@@ -82,7 +82,7 @@ if (botname === '__TEST__') {
 		// also flag for applying reverse entry too? would apply to branchspecs, not so simple for resolver
 		const propertiesToPropagate: [keyof NodeOptions, keyof EdgeOptions, ConfigBlendMode | null][] = [
 			['additionalSlackChannelForBlockages', 'additionalSlackChannel', null],
-			['postMessagesToAdditionalChannelOnly', 'postOnlyToAdditionalChannel', null],
+			['postMessagesToAdditionalChannelOnly', 'postOnlyToAdditionalChannel', 'override'],
 			['lastGoodCLPath', 'lastGoodCLPath', null],
 			['waitingForCISLink', 'waitingForCISLink', null],
 			['disallowSkip', 'disallowSkip', null],
