Override the channel setting not the only post to one channel setting

diff --git a/Engine/Extras/RoboMerge/v3/src/robo/branchgraph.ts b/Engine/Extras/RoboMerge/v3/src/robo/branchgraph.ts
--- a/Engine/Extras/RoboMerge/v3/src/robo/branchgraph.ts
+++ b/Engine/Extras/RoboMerge/v3/src/robo/branchgraph.ts
@@ -81,8 +81,8 @@ if (botname === '__TEST__') {
 	private propagateSourcePropertiesToEdges() {
 		// also flag for applying reverse entry too? would apply to branchspecs, not so simple for resolver
 		const propertiesToPropagate: [keyof NodeOptions, keyof EdgeOptions, ConfigBlendMode | null][] = [
-			['additionalSlackChannelForBlockages', 'additionalSlackChannel', null],
-			['postMessagesToAdditionalChannelOnly', 'postOnlyToAdditionalChannel', 'override'],
+			['additionalSlackChannelForBlockages', 'additionalSlackChannel', 'override'],
+			['postMessagesToAdditionalChannelOnly', 'postOnlyToAdditionalChannel', null],
 			['lastGoodCLPath', 'lastGoodCLPath', null],
 			['waitingForCISLink', 'waitingForCISLink', null],
 			['disallowSkip', 'disallowSkip', null],
