Fix approval message getting cut off when there is no owner

diff --git a/Engine/Extras/RoboMerge/v3/src/robo/nodebot.ts b/Engine/Extras/RoboMerge/v3/src/robo/nodebot.ts
--- a/Engine/Extras/RoboMerge/v3/src/robo/nodebot.ts
+++ b/Engine/Extras/RoboMerge/v3/src/robo/nodebot.ts
@@ -2234,7 +2234,7 @@ export class NodeBot extends PerforceStatefulBot implements NodeBotInterface {
 				{title: 'Change', short: true, value: makeClLink(pending.newCl)}
 			]
 
-			let message = slackAuthor ? `<@${slackAuthor}>` : `@${pending.change.author}` +
+			let message = (slackAuthor ? `<@${slackAuthor}>` : `@${pending.change.author}`) +
 						  "'s change " +
 						  (ownerEmail ? "owned by " + (slackOwner ? `<@${slackOwner}> ` : `@${pending.change.owner} `) : "") +
 						  `in ${pending.action.branch.name} needs to be approved.\n\n` +
