Adding more information about the exclusive checkout message so that the user knows how to proceed with the RoboMerge main page once the conflict is resolved.

diff --git a/Engine/Extras/RoboMerge/v3/src/robo/notifications.ts b/Engine/Extras/RoboMerge/v3/src/robo/notifications.ts
--- a/Engine/Extras/RoboMerge/v3/src/robo/notifications.ts
+++ b/Engine/Extras/RoboMerge/v3/src/robo/notifications.ts
@@ -612,7 +612,7 @@ export class BotNotifications implements BotEventHandler {
 						if (!exclusiveLockSlackUser) {
 							exclusiveLockSlackUser = `@${exclusiveLockUser.user} `
 						}
-						file.initial_comment = `${exclusiveLockSlackUser} please unlock the files blocking robomerge or work with ${blockage.owner} to resolve the conflict`
+						file.initial_comment = `${exclusiveLockSlackUser} please unlock the files blocking robomerge or work with ${blockage.owner} to resolve the conflict. Hit retry on the blocked stream once the files are unlocked.`
 					} else {
 						file.initial_comment = "The following locked files did not have their owner determined and as such those owners may not have been notified"
 					}
