Fix warning message for when server is running in preview mode

diff --git a/Engine/Extras/RoboMerge/v3/src/robo/roboserver.ts b/Engine/Extras/RoboMerge/v3/src/robo/roboserver.ts
--- a/Engine/Extras/RoboMerge/v3/src/robo/roboserver.ts
+++ b/Engine/Extras/RoboMerge/v3/src/robo/roboserver.ts
@@ -71,7 +71,7 @@ let ENVIRONMENT: {[param: string]: any}
 if (ENVIRONMENT.devMode) {
 	roboserverStartupLogger.warn('Running in DEV_MODE')
 }
-if (ENVIRONMENT.previewMode) {
+if (ENVIRONMENT.previewOnly) {
 	roboserverStartupLogger.warn('Running in PREVIEW_MODE')
 }
 
