Clean up no bots visible error messaging

diff --git a/Engine/Extras/RoboMerge/v3/public/js/boilerplate.js b/Engine/Extras/RoboMerge/v3/public/js/boilerplate.js
--- a/Engine/Extras/RoboMerge/v3/public/js/boilerplate.js
+++ b/Engine/Extras/RoboMerge/v3/public/js/boilerplate.js
@@ -1925,7 +1925,7 @@ if (!onLoginPage) {
 				$('.tags', $container).text(data.user.privileges && Array.isArray(data.user.privileges) ? ` (${data.user.privileges.join(', ')})` : '');
 
 				if (data.insufficientPrivelege) {
-					setErrorText('There are bots running but logged in user does not have admin access');
+					displayErrorMessage('There are bots running but logged in user does not have access to see any');
 				}
 
 			}
