Fix preview not displaying errors
Only show the graph controls on preview if there are no errors

diff --git a/Engine/Extras/RoboMerge/v3/public/js/preview.js b/Engine/Extras/RoboMerge/v3/public/js/preview.js
--- a/Engine/Extras/RoboMerge/v3/public/js/preview.js
+++ b/Engine/Extras/RoboMerge/v3/public/js/preview.js
@@ -191,8 +191,9 @@ function doit(query)
 			$('.bots').html([...botNames].map(s => `<tt>${s.toLowerCase()}</tt>`).join(', '))
 			$('#success-panel').show();
 			$('#control-panel').show();
+			$('#controlpanel').show();
 		},
-		error: function(xhr, error, status) {
+		error: error => {
 			const $errorPanel = $('#error-panel');
 			const errorMsg = error.responseText
 				? error.responseText.replace(/\t/g, '    ')
diff --git a/Engine/Extras/RoboMerge/v3/public/preview.html b/Engine/Extras/RoboMerge/v3/public/preview.html
--- a/Engine/Extras/RoboMerge/v3/public/preview.html
+++ b/Engine/Extras/RoboMerge/v3/public/preview.html
@@ -41,7 +41,7 @@
 				</div>
 			</div>
 		</div>
-		<details id="controlpanel" class="graphcontrolpanel">
+		<details id="controlpanel" class="graphcontrolpanel" style="display:none;">
 			<summary>Graph Control Panel</summary>
 			<div class="row p-1">
 				<div class="col-4">
