Query just a single stream rather than all streams when possible

diff --git a/Engine/Extras/RoboMerge/v3/src/common/perforce.ts b/Engine/Extras/RoboMerge/v3/src/common/perforce.ts
--- a/Engine/Extras/RoboMerge/v3/src/common/perforce.ts
+++ b/Engine/Extras/RoboMerge/v3/src/common/perforce.ts
@@ -631,6 +631,26 @@ export class PerforceContext {
 		return streams
 	}
 
+	async stream(streamName: string) {
+		const ztagResult = await this._execP4Ztag(null, ['stream', '-o', streamName], {quiet: false})
+		if (ztagResult.length > 2) {
+			throw new Error('Unexpected describe result')
+		}		
+
+		const stream: any = ztagResult[0];
+
+		if (stream.Update) {
+			const streamSpec: StreamSpec = {
+				stream: stream.Stream as string,
+				name: stream.Name as string,
+				parent: stream.Parent as string,
+				desc: stream.desc as string,
+			}
+			return streamSpec
+		}
+		return null
+	}
+
 	// find a workspace for the given user
 	// output format is an array of workspace names
 	async find_workspaces(user?: string, options?: {edgeServerAddress?: string, includeUnloaded?: boolean}) {
diff --git a/Engine/Extras/RoboMerge/v3/src/robo/gate.ts b/Engine/Extras/RoboMerge/v3/src/robo/gate.ts
--- a/Engine/Extras/RoboMerge/v3/src/robo/gate.ts
+++ b/Engine/Extras/RoboMerge/v3/src/robo/gate.ts
@@ -456,8 +456,7 @@ export class Gate {
 
 		const stream = await this.context.p4!.getStreamName(lastGoodCLPath);
 		if (typeof stream === 'string') {
-			const streams = await this.context.p4!.streams();
-			if (!streams.has(stream)) {
+			if (!(await this.context.p4!.stream(stream))) {
 				this.context.logger.error(`Unable to find stream '${stream}'`)
 				return null
 			}
diff --git a/Engine/Extras/RoboMerge/v3/src/robo/robo.ts b/Engine/Extras/RoboMerge/v3/src/robo/robo.ts
--- a/Engine/Extras/RoboMerge/v3/src/robo/robo.ts
+++ b/Engine/Extras/RoboMerge/v3/src/robo/robo.ts
@@ -424,8 +424,7 @@ async function init(logger: ContextualLogger) {
 			try {
 				const stream = await robo.p4.getStreamName(rootPath);
 				if (typeof stream === 'string') {
-					const streams = await robo.p4.streams();
-					if (streams.has(stream)) {
+					if (await robo.p4.stream(stream)) {
 						return stream
 					}
 				}
