Fix compile error and minor cleanup

diff --git a/Engine/Extras/RoboMerge/v3/src/robo/gate.ts b/Engine/Extras/RoboMerge/v3/src/robo/gate.ts
--- a/Engine/Extras/RoboMerge/v3/src/robo/gate.ts
+++ b/Engine/Extras/RoboMerge/v3/src/robo/gate.ts
@@ -482,7 +482,7 @@ export class Gate {
 		}
 	
 		const syncResult = await p4.sync(workspace, lastGoodCLPath, {opts:[P4_FORCE]})
-		if (syncResult.length != 1) {
+		if (!syncResult || syncResult.length != 1) {
 			this.context.logger.error(`Unable to sync ${lastGoodCLPath} in workspace ${workspace.name}`)
 			return null
 		}
diff --git a/Engine/Extras/RoboMerge/v3/src/robo/nodebot.ts b/Engine/Extras/RoboMerge/v3/src/robo/nodebot.ts
--- a/Engine/Extras/RoboMerge/v3/src/robo/nodebot.ts
+++ b/Engine/Extras/RoboMerge/v3/src/robo/nodebot.ts
@@ -757,11 +757,11 @@ export class NodeBot extends PerforceStatefulBot implements NodeBotInterface {
 			// We need the depot file for each conflict file to do comparisons
 			for (const remainingFile of remainingFiles) {
 				// Find target file 
-				const targetFileInDepotZtag = await this.p4.where(integratedCl.client, remainingFile.clientFile)
-				if (!targetFileInDepotZtag[0].depotFile) { 
+				const targetFileInDepot = await this.p4.where(integratedCl.client, remainingFile.clientFile)
+				if (!targetFileInDepot[0].depotFile) { 
 					throw new Error(`Error retrieving depot path for the merge target of ${remainingFile.clientFile}`)
 				}
-				remainingFile.targetDepotFile = targetFileInDepotZtag[0].depotFile
+				remainingFile.targetDepotFile = targetFileInDepot[0].depotFile
 			}
 
 			let remainingAllBinary = true
