Don't query for all streams and instead use checks against individual stream existence

diff --git a/Engine/Extras/RoboMerge/v3/src/robo/autobranchupdater.ts b/Engine/Extras/RoboMerge/v3/src/robo/autobranchupdater.ts
--- a/Engine/Extras/RoboMerge/v3/src/robo/autobranchupdater.ts
+++ b/Engine/Extras/RoboMerge/v3/src/robo/autobranchupdater.ts
@@ -151,7 +151,7 @@ export class AutoBranchUpdater implements Bot {
 		}
 
 		const validationErrors: string[] = []
-		const result = BranchDefs.parseAndValidate(validationErrors, branchGraphText, await this.p4.streams())
+		const result = await BranchDefs.parseAndValidate(this.p4, validationErrors, branchGraphText)
 		if (!result.branchGraphDef) {
 			// @todo email author of changes!
 			let errText = 'failed to parse/validate branch specs file\n'
diff --git a/Engine/Extras/RoboMerge/v3/src/robo/branchdefs.ts b/Engine/Extras/RoboMerge/v3/src/robo/branchdefs.ts
--- a/Engine/Extras/RoboMerge/v3/src/robo/branchdefs.ts
+++ b/Engine/Extras/RoboMerge/v3/src/robo/branchdefs.ts
@@ -1,12 +1,6 @@
 // Copyright Epic Games, Inc. All Rights Reserved.
 
-import { StreamSpecs } from '../common/perforce'
-
-// for validating all branchspecs
-// import { ContextualLogger } from '../common/logger';
-// import { initializePerforce, PerforceContext, StreamSpec } from '../common/perforce';
-// import fs = require('fs')
-// import path = require('path')
+import { PerforceContext } from '../common/perforce'
 
 const jsonlint: any = require('jsonlint')
 
@@ -269,7 +263,7 @@ export class BranchDefs {
 		}
 	}
 
-	static parseAndValidate(outErrors: string[], branchSpecsText: string, allStreamSpecs: StreamSpecs, isPreviewing?: boolean): ParseResult {
+	static async parseAndValidate(p4: PerforceContext, outErrors: string[], branchSpecsText: string, isPreviewing?: boolean): Promise<ParseResult> {
 		const defaultConfigForWholeBot: BotConfig = {
 			defaultStreamDepot: null,
 			defaultIntegrationMethod: null,
@@ -384,7 +378,7 @@ export class BranchDefs {
 				def.streamSubpath
 			)
 
-			if (!isPreviewing && streamResult.stream && !allStreamSpecs.has(streamResult.stream)) {
+			if (!isPreviewing && streamResult.stream && !await p4.stream(streamResult.stream)) {
 				outErrors.push(`Stream ${streamResult.stream} not found`)
 			}
 		}
@@ -541,34 +535,3 @@ export class BranchDefs {
 		return {branchGraphDef: branchGraph, config: defaultConfigForWholeBot}
 	}
 }
-
-// function verifyAllBranchmaps(allStreamSpecs: Map<string, StreamSpec>, folder: string) {
-
-// 	for (const file of fs.readdirSync(folder, {encoding: "utf8"})) {
-// 		if (!file.endsWith('branchmap.json') || file.indexOf('iron.') >= 0) {
-// 			continue
-// 		}
-// 		console.log(file)
-// 		const validationErrors: string[] = []
-
-// 		const result = BranchDefs.parseAndValidate(
-// 			validationErrors,
-// 			fs.readFileSync(path.join(folder, file), 'utf8'),
-// 			allStreamSpecs)
-
-// 		if (!result.branchGraphDef) {
-// 			throw new Error(validationErrors.length === 0 ? 'Failed to parse' : validationErrors.join('\n'))
-// 		}
-// 		console.log(`Branches found in ${file}: ${result.branchGraphDef.branches.length}`)
-// 	}
-// }
-
-// export async function runTests(logger: ContextualLogger) {
-// 	if (process.platform === 'darwin') {
-// 		await initializePerforce(logger)
-// 		const allStreamSpecs = await (new PerforceContext(logger)).streams()
-
-// 		verifyAllBranchmaps(allStreamSpecs, '../RoboMerge/data')
-// 	}
-// 	return 0
-// }
diff --git a/Engine/Extras/RoboMerge/v3/src/robo/graphbot.ts b/Engine/Extras/RoboMerge/v3/src/robo/graphbot.ts
--- a/Engine/Extras/RoboMerge/v3/src/robo/graphbot.ts
+++ b/Engine/Extras/RoboMerge/v3/src/robo/graphbot.ts
@@ -4,7 +4,7 @@ import * as Sentry from '@sentry/node';
 import { ContextualLogger } from '../common/logger';
 import { Mailer } from '../common/mailer';
 import * as p4util from '../common/p4util';
-import { PerforceContext, Workspace, StreamSpecs } from '../common/perforce';
+import { PerforceContext, Workspace } from '../common/perforce';
 import { AutoBranchUpdater } from './autobranchupdater';
 import { bindBadgeHandler } from './badges';
 import { Bot } from './bot-interfaces';
@@ -44,39 +44,44 @@ export class GraphBot implements GraphInterface, BotEventHandler {
 
 	private p4: PerforceContext;
 
-	constructor(botname: string, private mailer: Mailer, private externalUrl: string, allStreamSpecs: StreamSpecs) {
+	private constructor(private mailer: Mailer, private externalUrl: string) {
+	}
+
+	static async CreateAsync(botname: string, mailer: Mailer, externalUrl: string) {
+		let graphBot = new GraphBot(mailer, externalUrl)
+
 		if (!GraphBot.dataDirectory) {
 			throw new Error('Data directory must be set before creating a BranchGraph')
 		}
 
-		this.botLogger = new ContextualLogger(botname.toUpperCase())
-		this.p4 = new PerforceContext(this.botLogger)
+		graphBot.botLogger = new ContextualLogger(botname.toUpperCase())
+		graphBot.p4 = new PerforceContext(graphBot.botLogger)
 
-		this.branchGraph = new BranchGraph(botname)
-		this.filename = botname + '.branchmap.json'
+		graphBot.branchGraph = new BranchGraph(botname)
+		graphBot.filename = botname + '.branchmap.json'
 
-		const branchSettingsPath = `${GraphBot.dataDirectory}/${this.filename}`
+		const branchSettingsPath = `${GraphBot.dataDirectory}/${graphBot.filename}`
 
-		this.botLogger.info(`Loading branch map from ${branchSettingsPath}`)
+		graphBot.botLogger.info(`Loading branch map from ${branchSettingsPath}`)
 		const fileText = require('fs').readFileSync(branchSettingsPath, 'utf8')
 
 		const validationErrors: string[] = []
-		const result = BranchDefs.parseAndValidate(validationErrors, fileText, allStreamSpecs)
+		const result = await BranchDefs.parseAndValidate(graphBot.p4, validationErrors, fileText)
 		if (!result.branchGraphDef) {
 			throw new Error(validationErrors.length === 0 ? 'Failed to parse' : validationErrors.join('\n'))
 		}
 
-		this.branchGraph.config = result.config
+		graphBot.branchGraph.config = result.config
 
 		let error: string | null = null
 		try {
-			this.branchGraph._initFromBranchDefInternal(result.branchGraphDef)
+			graphBot.branchGraph._initFromBranchDefInternal(result.branchGraphDef)
 		}
 		catch (exc) {
 			// reset - don't keep a partially configured bot around
-			this.branchGraph = new BranchGraph(botname)
-			this.branchGraph.config = result.config
-			this.branchGraph._initFromBranchDefInternal(null)
+			graphBot.branchGraph = new BranchGraph(botname)
+			graphBot.branchGraph.config = result.config
+			graphBot.branchGraph._initFromBranchDefInternal(null)
 			error = exc.toString();
 		}
 
@@ -86,12 +91,14 @@ export class GraphBot implements GraphInterface, BotEventHandler {
 
 		// start empty bot on error - can be fixed up by branch definition check-in
 		if (error) {
-			this.botLogger.error(`Problem starting up bot ${botname}: ${error}`);
+			graphBot.botLogger.error(`Problem starting up bot ${botname}: ${error}`);
 		}
 
-		this.settings = new Settings(botname, this.branchGraph, this.botLogger, this.p4)
+		graphBot._settings = new Settings(botname, graphBot.branchGraph, graphBot.botLogger, graphBot.p4)
 
-		this.externalUrl = externalUrl
+		graphBot.externalUrl = externalUrl
+
+		return graphBot
 	}
 
 	findNode(branchname: BranchArg): NodeBot | undefined {
@@ -495,7 +502,11 @@ export class GraphBot implements GraphInterface, BotEventHandler {
 		}
 	}
 
-	readonly settings: Settings
+	private _settings: Settings
+
+	get settings() {
+		return this._settings;
+	}
 	private botlist: Bot[] = []
 	private waitTime?: number
 
diff --git a/Engine/Extras/RoboMerge/v3/src/robo/preview.ts b/Engine/Extras/RoboMerge/v3/src/robo/preview.ts
--- a/Engine/Extras/RoboMerge/v3/src/robo/preview.ts
+++ b/Engine/Extras/RoboMerge/v3/src/robo/preview.ts
@@ -30,8 +30,6 @@ export async function getPreview(cl: number, singleBot?: string) {
 		}
 	}
 
-	const allStreamSpecs = await p4.streams()
-
 	const errors = []
 	for (const [bot, path] of bots) {
 		if (singleBot && bot.toLowerCase() !== singleBot.toLowerCase()) {	
@@ -40,7 +38,7 @@ export async function getPreview(cl: number, singleBot?: string) {
 		const fileText = await p4.print(`${path}@=${cl}`)
 
 		let validationErrors: string[] = []
-		const result = BranchDefs.parseAndValidate(validationErrors, fileText, allStreamSpecs, true)
+		const result = await BranchDefs.parseAndValidate(p4, validationErrors, fileText, true)
 
 		const errorPrefix = `\n\t${bot} validation failed: `
 		if (!result.branchGraphDef) {
diff --git a/Engine/Extras/RoboMerge/v3/src/robo/robo.ts b/Engine/Extras/RoboMerge/v3/src/robo/robo.ts
--- a/Engine/Extras/RoboMerge/v3/src/robo/robo.ts
+++ b/Engine/Extras/RoboMerge/v3/src/robo/robo.ts
@@ -11,7 +11,7 @@ import { Arg, readProcessArgs } from '../common/args';
 import { _setTimeout } from '../common/helper';
 import { ContextualLogger } from '../common/logger';
 import { Mailer } from '../common/mailer';
-import { ClientSpec, getPerforceUsername, getRootDirectoryForBranch, initializePerforce, PerforceContext, StreamSpecs } from '../common/perforce';
+import { ClientSpec, getPerforceUsername, getRootDirectoryForBranch, initializePerforce, PerforceContext } from '../common/perforce';
 import { BuildVersion, VersionReader } from '../common/version';
 import { CertFiles } from '../common/webserver';
 import { addBranchGraph, Graph, GraphAPI } from '../new/graph';
@@ -498,7 +498,7 @@ async function init(logger: ContextualLogger) {
 	}
 
 	_initMailer(logger)
-	_initGraphBots(await robo.p4.streams(), logger)
+	await _initGraphBots(logger)
 	if (!DEBUG_SKIP_BRANCH_SETUP) {
 		await _initBranchWorkspacesForAllBots(logger)
 	}
@@ -632,10 +632,10 @@ function shutdown(exitCode: number, logger: ContextualLogger) {
 process.once('SIGINT', () => { roboStartupLogger.error("Caught SIGINT"); shutdown(2, roboStartupLogger); });
 process.once('SIGTERM', () => { roboStartupLogger.error("Caught SIGTERM"); shutdown(0, roboStartupLogger); });
 
-function _initGraphBots(allStreamSpecs: StreamSpecs, logger: ContextualLogger) {
+async function _initGraphBots(logger: ContextualLogger) {
 	for (const botname of args.botname)	{
 		logger.info(`Initializing bot ${botname}`)
-		const graphBot = new GraphBot(botname, robo.mailer, args.externalUrl, allStreamSpecs)
+		const graphBot = await GraphBot.CreateAsync(botname, robo.mailer, args.externalUrl)
 		robo.graphBots.set(graphBot.branchGraph.botname, graphBot)
 
 		graphBot.reloadAsyncListeners.add(_onBranchSpecReloaded)
