Convert robomerge to use -ztag -Mj in most cases. resolve and submit still using traditional ztag format until more time can be spent to untangle the differences in the return format

diff --git a/Engine/Extras/RoboMerge/v3/functional_tests/framework/src/test-perforce.ts b/Engine/Extras/RoboMerge/v3/functional_tests/framework/src/test-perforce.ts
--- a/Engine/Extras/RoboMerge/v3/functional_tests/framework/src/test-perforce.ts
+++ b/Engine/Extras/RoboMerge/v3/functional_tests/framework/src/test-perforce.ts
@@ -1,8 +1,6 @@
 // Copyright Epic Games, Inc. All Rights Reserved.
 import { ChildProcess, execFile, ExecFileOptions } from 'child_process'
 
-import * as ztag from './ztag'
-
 class ExecOpts {
 	stdin?: string
 }
@@ -18,6 +16,13 @@ export interface DescribeResult {
 	rev: number
 }
 
+type Type = 'string' | 'integer' | 'boolean'
+
+type ParseOptions = {
+	expected?: {[field: string]: Type}
+	optional?: {[field: string]: Type}
+}
+
 export const VERBOSE = false
 export class Perforce {
 	static readonly P4_EXE = process.platform === 'win32' ? 'p4.exe' : 'p4'
@@ -86,30 +91,167 @@ export class Perforce {
 	}
 
 	async changes(client: string, stream: string, limit: number, pending?: boolean) {
-		const ztagResult = await this.exec(['-ztag', '-c', client, 'changes', '-l', '-m', limit.toString(), '-s', pending ? 'pending' : 'submitted', stream + '/...'])
-
-		const parseResult = ztag.parseZtagOutput(ztagResult, 0, {expected: {change: 'integer', desc: 'string'}})
+		const jsonResult = await this.execAndParse(['-c', client, 'changes', '-l', '-m', limit.toString(), '-s', pending ? 'pending' : 'submitted', stream + '/...'], undefined, {expected: {change: 'integer', desc: 'string'}})
 		const result: Change[] = []
-		for (const entry of parseResult) {
+		for (const entry of jsonResult) {
 			result.push({change: entry.change as number, description: entry.desc as string})
 		}
 		return result
 	}
 
+	private describeHeaderExpectedShape: ParseOptions = {
+		expected: {change: 'integer', user: 'string', client: 'string'}
+	}
+
+	private describeEntryExpectedShape: ParseOptions = {
+		expected: {action: 'string', rev: 'integer', depotFile: 'string'}
+	}
+	
 	async describe(cl: number) {
-		const ztagResult = await this.exec(['-ztag', 'describe', cl.toString()])
-		const parseResult = ztag.parseHeaderAndArray(ztagResult, 0,
-			{expected: {change: 'integer', user: 'string', client: 'string'}},
-			{expected: {action: 'string', rev: 'integer', depotFile: 'string'}}
-		)
-
-		const result: DescribeResult[] = []
-		for (const entry of parseResult.slice(1)) {
-			result.push({file: entry.depotFile as string, action: entry.action as string, rev: entry.rev as number})
+		let result = (await this.execAndParseArray(['describe', cl.toString()], undefined, this.describeHeaderExpectedShape, this.describeEntryExpectedShape))[0]
+
+		return result.entries as DescribeResult[];
+	}
+
+	private parseValue(key: string, value: any, parseOptions?: ParseOptions)
+	{
+		if (parseOptions) {
+			const optionalType = parseOptions.optional && parseOptions.optional[key]
+			const fieldType = optionalType || (parseOptions.expected && parseOptions.expected[key]) || 'string'
+			if (fieldType === 'boolean') {
+				if (!optionalType || value) {
+					const valLower = value.toLowerCase()
+					if (valLower !== 'true' && valLower !== 'false') {
+						throw new Error(`Failed to parse boolean field ${key}, value: ${value}`)
+					}
+					return valLower === 'true'
+				}
+				return undefined
+			}
+			else if (fieldType === 'integer') {
+				// ignore empty strings for optional fields (e.g. p4.changes can return a 'shelved' property with no value)
+				if (!optionalType || value) {
+					const num = parseInt(value)
+					if (isNaN(num)) {
+						throw new Error(`Failed to parse number field ${key}, value: ${value}`)
+					}
+					return num
+				}
+				return undefined
+			}
+		}
+		return value		
+	}
+
+	private async execAndParse(args: string[], execOptions?: ExecOpts, parseOptions?: ParseOptions) {
+
+		args = ['-ztag', '-Mj', ...args]
+		let rawResult = await this.exec(args, execOptions)
+
+		let result = []
+		let startIndex = 0;
+
+		let reviver = (key: string, value: any) => {
+			return this.parseValue(key, value, parseOptions)
 		}
+
+		while(startIndex < rawResult.length) {
+			const endIndex = rawResult.indexOf('}\n', startIndex)
+			const parsedResult = JSON.parse(rawResult.slice(startIndex, endIndex != -1 ? endIndex+1 : undefined), reviver)
+			for (const expected in ((parseOptions && parseOptions.expected) || []))
+			{
+				if (!parsedResult[expected])
+				{
+					throw new Error(`Expected field ${expected} not present in ${JSON.stringify(parsedResult)}`)
+				}
+			}
+			result.push(parsedResult)
+			startIndex = endIndex + 2
+		}
+
+		if (result.length == 1 && Object.hasOwn(result[0],'data') && Object.hasOwn(result[0],'generic') && Object.hasOwn(result[0],'severity')) {
+			const cmd = `p4 ${[...args, ...(execOptions && execOptions.stdin ? ['-i'] : [])].join(' ')}`
+			throw [new Error(`P4 Error: ${cmd}\n${result[0]['data']}`), result[0]['data'].replace(Perforce.REGEX_NEWLINE, '\n')]
+		}
+
 		return result
 	}
 
+	private async execAndParseArray(args: string[], execOptions?: ExecOpts, headerOptions?: ParseOptions, arrayEntryOptions?: ParseOptions) {
+
+		args = ['-ztag', '-Mj', ...args]
+		let rawResult = await this.exec(args, execOptions)
+
+		let result = []
+		let startIndex = 0;
+
+		let reviver = (key: string, value: any) => {
+			const arrayElementMatch = key.match(/^(.*?)\d+$/)
+			if (arrayElementMatch) {
+				return this.parseValue(arrayElementMatch[1], value, arrayEntryOptions)
+			}
+			return this.parseValue(key, value, headerOptions)
+		}
+
+		while(startIndex < rawResult.length) {
+			const endIndex = rawResult.indexOf('}\n', startIndex)
+			let parsedResult = JSON.parse(rawResult.slice(startIndex, endIndex != -1 ? endIndex+1 : undefined), reviver)
+			for (let expected in headerOptions && headerOptions.expected || []) {
+				if (!parsedResult[expected])
+				{
+					throw new Error(`Expected field ${expected} not present in ${JSON.stringify(parsedResult)}`)
+				}
+			}
+			let organizedResult: {[key:string]:any} = {};
+			organizedResult.entries = []
+			let expectedCounts: number[] = []
+			for (const field in parsedResult) {
+				const arrayElementMatch = field.match(/^(.*?)(\d+)$/)
+				if (arrayElementMatch) {
+					const arrayField = arrayElementMatch[1]
+					const arrayIndex = parseInt(arrayElementMatch[2])
+					const curLength = expectedCounts.length
+					if (arrayIndex >= curLength) {
+						organizedResult.entries.length = arrayIndex+1
+						expectedCounts.length = arrayIndex+1
+						for (let i=curLength; i<organizedResult.entries.length; i++) {
+							let newEntry: {[key:string]:any} = {};
+							organizedResult.entries[i] = newEntry
+						}
+					}
+					organizedResult.entries[arrayIndex][arrayField] = parsedResult[field]
+					if (arrayEntryOptions && arrayEntryOptions.expected && arrayEntryOptions.expected[arrayField]) {
+						expectedCounts[arrayIndex] += 1
+					}
+				}
+				else {
+					organizedResult[field] = parsedResult[field]
+				}
+			}
+			if (arrayEntryOptions && arrayEntryOptions.expected) {
+				const expectedCount = Object.keys(arrayEntryOptions.expected).length
+				for (let index in expectedCounts) {
+					if (expectedCounts[index] != expectedCount) {
+						for (let expected in arrayEntryOptions.expected) {
+							if (!organizedResult.entries[index][expected])
+							{
+								throw new Error(`Expected field ${expected} not present in ${JSON.stringify(organizedResult.entries[index])}`)
+							}
+						}
+					}
+				}
+			}
+			result.push(organizedResult)
+			startIndex = endIndex + 2
+		}
+
+		if (result.length == 1 && Object.hasOwn(result[0],'data') && Object.hasOwn(result[0],'generic') && Object.hasOwn(result[0],'severity')) {
+			const cmd = `p4 ${[...args, ...(execOptions && execOptions.stdin ? ['-i'] : [])].join(' ')}`
+			throw [new Error(`P4 Error: ${cmd}\n${result[0]['data']}`), result[0]['data'].replace(Perforce.REGEX_NEWLINE, '\n')]
+		}
+
+		return result
+	}
 
 	private exec(args: string[], optOpts?: ExecOpts): Promise<string> {
 		if (VERBOSE) console.log('Running: ' + args.join(' '))
@@ -176,4 +318,5 @@ export class Perforce {
 
 		return execPromise
 	}
+
 }
diff --git a/Engine/Extras/RoboMerge/v3/functional_tests/framework/src/ztag.ts b/Engine/Extras/RoboMerge/v3/functional_tests/framework/src/ztag.ts
deleted file mode 100644
diff --git a/Engine/Extras/RoboMerge/v3/src/common/perforce.ts b/Engine/Extras/RoboMerge/v3/src/common/perforce.ts
--- a/Engine/Extras/RoboMerge/v3/src/common/perforce.ts
+++ b/Engine/Extras/RoboMerge/v3/src/common/perforce.ts
@@ -7,8 +7,12 @@ import { roboAnalytics } from '../robo/roboanalytics';
 import { ContextualLogger, NpmLogLevel } from './logger';
 import { VersionReader } from './version';
 
-import * as ztag from './ztag'
+type Type = 'string' | 'integer' | 'boolean'
 
+type ParseOptions = {
+	expected?: {[field: string]: Type}
+	optional?: {[field: string]: Type}
+}
 const p4exe = process.platform === 'win32' ? 'p4.exe' : 'p4';
 const RETRY_ERROR_MESSAGES = [
 	'socket: Connection reset by peer',
@@ -24,11 +28,15 @@ export const EXCLUSIVE_CHECKOUT_REGEX = INTEGRATION_FAILURE_REGEXES[0][0]
 
 const REVERT_FAILURE_DUE_TO_MOVE_REGEX: RegExp = /(.*)#[0-9]+ - has been moved, not reverted/
 
-const changeResultExpectedShape: ztag.ParseOptions = {
+const changeResultExpectedShape: ParseOptions = {
 	expected: {change: 'integer', client: 'string', user: 'string', desc: 'string', time: 'integer', status: 'string', changeType: 'string'},
 	optional: {oldChange: 'integer'}
 }
 
+const describeEntryExpectedShape: ParseOptions = {
+	optional: { depotFile: 'string', action: 'string', rev: 'integer', type: 'string' }
+}
+
 const ztag_group_rex = /\n\n\.\.\.\s/;
 const ztag_field_rex = /(?:\n|^)\.\.\.\s/;
 const newline_rex = /\r\n|\n|\r/g;
@@ -190,31 +198,6 @@ export function parseZTag(buffer: string, opts?: ExecZtagOpts) {
 	return output;
 }
 
-function readObjectListFromZtagOutput(obj: any, keysToLookFo: string[], logger: ContextualLogger) {
-	const result: any = []
-
-	for (let index = 0; ; ++index) {
-		// must be higher than file limit!
-		if (index === 50000) {
-			logger.warn('Parse -ztag WARNING: broke out after 50000 items')
-			break
-		}
-
-		const item: any = {}
-		for (const lookFor of keysToLookFo) {
-			const key = lookFor + index
-			if (obj[key]) {
-				item[lookFor] = obj[key]
-			}
-		}
-		if (Object.keys(item).length === 0)
-			break
-
-		result.push(item)
-	}
-	return result
-}
-
 class CommandRecord {
 	constructor(public cmd: string, public start: Date = new Date()) {
 	}
@@ -387,7 +370,7 @@ export class ResolveResult {
 
 				["clientFile", "fromFile", "startFromRev", "endFromRev", "resolveType", "resolveFlag"].forEach(function (keyValue) {
 					if (!ztagGroup[keyValue]) {
-						throw new Error(`Resolve output missing ${keyValue} in ztag: ${ztagGroup.toString()}`)
+						throw new Error(`Resolve output missing ${keyValue} in ztag: ${JSON.stringify(ztagGroup)}`)
 					}
 				})
 
@@ -480,13 +463,13 @@ export async function initializePerforce(logger: ContextualLogger) {
 		}
 	}, 5*60*1000)
 
-	const output = await PerforceContext._execP4Ztag(logger, null, ["login", "-s"], { noUsername: true });
+	const output = await PerforceContext.execAndParse(logger, null, ["login", "-s"], { noUsername: true });
 	let resp = output[0];
 
 	if (resp && resp.User) {
 		perforceUsername = resp.User;
 
-		const serversOutput = await PerforceContext._execP4Ztag(logger, null, ["servers"]);
+		const serversOutput = await PerforceContext.execAndParse(logger, null, ["servers"]);
 		perforceMultiServerEnvironment = serversOutput.length > 1
 	}
 }
@@ -513,21 +496,17 @@ export class PerforceContext {
 			args = [...args, '-c', workspaceName]
 		}
 
-		return this._execP4Ztag(null, args, { multiline: true, edgeServerAddress });
+		return this.execAndParse(null, args, { edgeServerAddress });
 	}
 
 	/** get a single change and return it in the format of changes() */
 	async getChange(changenum: number) {
-		let result = await this._execP4Ztag(null, ['change', '-o', changenum.toString()], 
-						{format: '{"change":%Change%,"client":"%Client%","user":"%User%","status":"%Status%","desc":"%Description%"}'})
-		result = result.trimEnd().replaceAll("\\","\\\\").replaceAll("\n","\\n").replaceAll("\t","\\t")
-		
-		// Have to escape any quotes in the description otherwise the JSON parsing fails
-		const descIndex = result.search('"desc":')
-		const descEndIndex = result.lastIndexOf('"')	
-		result = result.slice(0,descIndex+8) + result.slice(descIndex+8,descEndIndex).replaceAll('"','\\"') + '"}'	
-		
-		return JSON.parse(result);
+		let result = (await this.execAndParse(null, ['change', '-o', changenum.toString()]))[0]
+		result.change = parseInt(result.Change)
+		result.client = result.Client
+		result.desc = result.Description
+		result.user = result.User
+		return result
 	}
 
 	/**
@@ -565,7 +544,7 @@ export class PerforceContext {
 			this.logger.warn(`p4.latestChange took ${durationSeconds}s`)
 		}
 
-		return result[0] as unknown as Change
+		return result[0] as Change
 	}
 
 	changesBetween(path: string, from: number, to: number) {
@@ -574,9 +553,9 @@ export class PerforceContext {
 	}
 
 	async getDepot(depotName: string) {
-		if ((await this._execP4Ztag(null, ['depots', '-e', depotName])).length > 0)
+		if ((await this.execAndParse(null, ['depots', '-e', depotName])).length > 0)
 		{
-			return this._execP4Ztag(null, ['depot', '-o', depotName], {reduce: true})
+			return (await this.execAndParse(null, ['depot', '-o', depotName]))[0]
 		}
 		return null
 	}
@@ -589,6 +568,9 @@ export class PerforceContext {
 		}
 
 		const depot = await this.getDepot(path.substring(2,depotEndChar))
+		if (!depot) {
+			throw new Error(`Unable to find ${depot}`)
+		}
 		if (depot.Type == 'stream') {
 			let streamDepth = depot.StreamDepth.match(/\//g).length - 2
 			if (streamDepth < 1) {
@@ -602,22 +584,18 @@ export class PerforceContext {
 						return path
 					}
 					else {
-						return new Error(`Unable to determine stream from $(path): not enough depth`)
+						return new Error(`Unable to determine stream from ${path}: not enough depth`)
 					}
 				}
 				streamNameEnd = path.indexOf("/",streamNameEnd+1)
 			}
 			return path.substring(0,streamNameEnd)
 		}
-		return new Error(`Depot $(depot) is not of type stream`)
+		return new Error(`Depot ${depot} is not of type stream`)
 	}
 
 	async streams() {
-		const rawStreams = await this.execAndParse(null, ['streams'], {quiet: true}, {
-			expected: {Stream: 'string', Update: 'integer', Access: 'integer', Owner: 'string', Name: 'string', Parent: 'string', Type: 'string', desc: 'string',
-				Options: 'string', firmerThanParent: 'string', changeFlowsToParent: 'boolean', changeFlowsFromParent: 'boolean', baseParent: 'string'}
-		})
-
+		const rawStreams = await this.execAndParse(null, ['streams'], {quiet: true})
 		const streams = new Map<string, StreamSpec>()
 		for (const raw of rawStreams) {
 			const stream: StreamSpec = {
@@ -632,14 +610,8 @@ export class PerforceContext {
 	}
 
 	async stream(streamName: string) {
-		const ztagResult = await this._execP4Ztag(null, ['stream', '-o', streamName], {quiet: false})
-		if (ztagResult.length > 2) {
-			throw new Error('Unexpected describe result')
-		}		
-
-		const stream: any = ztagResult[0];
-
-		if (stream.Update) {
+		try {
+			const stream = (await this.execAndParse(null, ['streams', streamName], {quiet: false}))[0]
 			const streamSpec: StreamSpec = {
 				stream: stream.Stream as string,
 				name: stream.Name as string,
@@ -648,7 +620,10 @@ export class PerforceContext {
 			}
 			return streamSpec
 		}
-		return null
+		catch
+		{
+			return null
+		}
 	}
 
 	// find a workspace for the given user
@@ -666,15 +641,15 @@ export class PerforceContext {
 			args.push('-a')
 		}
 
-		let opts: ExecZtagOpts = { multiline: true }
+		let opts: ExecOpts = {}
 		if (edgeServerAddress && edgeServerAddress !== 'commit') {
 			opts.edgeServerAddress = edgeServerAddress
 		}
 
 		let workspaces = [];
 		try {
-			let parsedLoadedClients = this._execP4Ztag(null, args, opts);
-			let parsedUnloadedClients = (includeUnloaded ? this._execP4Ztag(null, [...args, '-U'], opts) : null)
+			let parsedLoadedClients = this.execAndParse(null, args, opts);
+			let parsedUnloadedClients = (includeUnloaded ? this.execAndParse(null, [...args, '-U'], opts) : null)
 			for (let clientDef of await parsedLoadedClients) {
 				if (clientDef.client) {
 					workspaces.push(clientDef);
@@ -709,7 +684,7 @@ export class PerforceContext {
 	}
 
 	find_workspace_by_name(workspaceName: string) {
-		return this._execP4Ztag(null, ['clients', '-E', workspaceName]);
+		return this.execAndParse(null, ['clients', '-E', workspaceName]);
 	}
 
 	async getWorkspaceEdgeServer(workspaceName: string): Promise<{id: string, address: string} | null> {
@@ -726,7 +701,7 @@ export class PerforceContext {
 	}
 
 	reloadWorkspace(workspaceName: string, edgeServerAddress?: string) {
-		return this._execP4Ztag(null, ['reload', '-c', workspaceName], {edgeServerAddress});
+		return this.execAndParse(null, ['reload', '-c', workspaceName], {edgeServerAddress});
 	}
 
 	getEdgeServerAddress(serverId: string) {
@@ -764,7 +739,7 @@ export class PerforceContext {
 		}
 		args.push(depotPath)
 		try {
-			return await this._execP4Ztag(workspace, args)
+			return await this.execAndParse(workspace, args)
 		}
 		catch (reason) {
 			if (!isExecP4Error(reason)) {
@@ -965,7 +940,7 @@ export class PerforceContext {
 		// execute the P4 command
 		let changes;
 		try {
-			changes = await this._execP4Ztag(workspace, cmdList, { numRetries: 0, edgeServerAddress: opts.edgeServerAddress });
+			changes = await this.execAndParse(workspace, cmdList, { numRetries: 0, edgeServerAddress: opts.edgeServerAddress });
 		}
 		catch (reason) {
 			if (!isExecP4Error(reason)) {
@@ -1226,7 +1201,7 @@ export class PerforceContext {
 			// see which workspace has a file checked out/added
 			args.push(exclusive && perforceMultiServerEnvironment ? '-x' : '-a', arg)
 		}
-		return this._execP4Ztag(workspace, args) as Promise<OpenedFileRecord[]>
+		return this.execAndParse(workspace, args) as Promise<OpenedFileRecord[]>
 	}
 
 	async revertFile(file: string) {
@@ -1246,7 +1221,7 @@ export class PerforceContext {
 		const edgeServer = await this.getWorkspaceEdgeServer(client)
 		const args = ['revert', '-C', client, ...files]
 		try {
-			const results = await this._execP4Ztag(null, args, {edgeServerAddress: edgeServer?.address})
+			const results = await this.execAndParse(null, args, {edgeServerAddress: edgeServer?.address})
 
 			// Files can fail to be reverted because of moves, lets try and resolve that
 			const movedFiles = 
@@ -1258,7 +1233,7 @@ export class PerforceContext {
 								       .map(match => match![1])
 							), [])
 			if (movedFiles.length > 0) {
-				const openedFiles = await this._execP4Ztag(null, ['opened','-a', ...movedFiles])
+				const openedFiles = await this.execAndParse(null, ['opened','-a', ...movedFiles])
 				const pairedAdds: string[] = openedFiles.map((openedFile: any) => openedFile.movedFile)
 				await this.revertFiles(pairedAdds, client)
 			}
@@ -1303,7 +1278,7 @@ export class PerforceContext {
 	async listFilesToResolve(roboWorkspace: RoboWorkspace, changelist: number) {
 		const workspace = coercePerforceWorkspace(roboWorkspace);
 		try {
-			return await this._execP4Ztag(workspace, ['resolve', '-n', '-c', changelist.toString()]);
+			return await this.execAndParse(workspace, ['resolve', '-n', '-c', changelist.toString()]);
 		}
 		catch (err) {
 			if (err.toString().toLowerCase().includes('no file(s) to resolve')) {
@@ -1419,52 +1394,27 @@ export class PerforceContext {
 			...(maxFiles ? ['-m', maxFiles.toString()] : []),
 			...(includeShelved ? ['-S'] : []),
 			cl.toString()]
-		const ztagResult = await this._execP4Ztag(null, args, { multiline:true });
-
-		if (ztagResult.length > 2) {
-			throw new Error('Unexpected describe result')
-		}
-
-		const clInfo: any = ztagResult[0];
-		const result: DescribeResult = {
-			user: clInfo.user || '',
-			status: clInfo.status || '',
-			description: clInfo.desc || '',
-			path: clInfo.path || ztagResult[1].path,
-			date: clInfo.time ? new Date(clInfo.time * 1000) : null,
-			entries: []
-		};
-
-		for (const obj of readObjectListFromZtagOutput(ztagResult[ztagResult.length === 2 ? 1 : 0], ['depotFile', 'action', 'rev', 'type'], this.logger)) {
-			let rev = -1;
-			if (obj.rev) {
-				const num = parseInt(obj.rev);
-				if (!isNaN(num)) {
-					rev = num;
-				}
-			}
-			result.entries.push({
-				depotFile: obj.depotFile || '',
-				action: obj.action || '',
-				rev,
-				type: obj.type || ''
-			});
-		}
-		return result;
+		let result = (await this.execAndParseArray(null, args, undefined, undefined, describeEntryExpectedShape))[0]
+		result.user = result.user || ''
+		result.status = result.status || ''
+		result.description = result.desc || ''
+		result.date = result.time ? new Date(result.time * 1000) : null
+
+		return result as DescribeResult;
 	}
 
 	async dirs(path: string) {
 		if (path.endsWith("/...")) {
 			path = path.substring(0,path.length-3) + "*"
 		}
-		return (await this._execP4Ztag(null, ['dirs', path], { multiline:true })).map((dir: any) => dir.dir)
+		return (await this.execAndParse(null, ['dirs', path])).map((dir: any) => dir.dir)
 	}
 
 	async files(path: string, maxFiles?: number) {
 		const args = ['files',
 			...(maxFiles ? ['-m', maxFiles.toString()] : []),
 			path]
-		return await this._execP4Ztag(null, args, { multiline:true });
+		return await this.execAndParse(null, args);
 	}
 
 	// update the fields on an existing CL using p4 change
@@ -1513,7 +1463,7 @@ export class PerforceContext {
 	}
 
 	where(roboWorkspace: RoboWorkspace, clientPath: string) {
-		return this._execP4Ztag(roboWorkspace, ['where', clientPath])
+		return this.execAndParse(roboWorkspace, ['where', clientPath])
 	}
 
 	async filelog(roboWorkspace: RoboWorkspace, depotPath: string, beginRev: string, endRev: string, longOutput = false) {
@@ -1524,7 +1474,7 @@ export class PerforceContext {
 		args.push(`${depotPath}#${beginRev},${endRev}`)
 
 		try {
-			return await this._execP4Ztag(roboWorkspace, args, { multiline: longOutput })
+			return await this.execAndParse(roboWorkspace, args)
 		}
 		catch (reason) {
 			if (!isExecP4Error(reason)) {
@@ -1543,7 +1493,7 @@ export class PerforceContext {
 
 	async fstat(roboWorkspace: RoboWorkspace, depotPath: string) {
 		try {
-			return await this._execP4Ztag(roboWorkspace, ['fstat', depotPath])
+			return await this.execAndParse(roboWorkspace, ['fstat', depotPath])
 		}
 		catch (reason) {
 			if (!isExecP4Error(reason)) {
@@ -1573,7 +1523,7 @@ export class PerforceContext {
 		args.push(depotPath)
 
 		try {
-			return await this._execP4Ztag(roboWorkspace, args)
+			return await this.execAndParse(roboWorkspace, args)
 		}
 		catch (reason) {
 			if (!isExecP4Error(reason)) {
@@ -1609,14 +1559,7 @@ export class PerforceContext {
 	}
 
 	// execute a perforce command
-	static _execP4(logger: ContextualLogger, roboWorkspace: RoboWorkspace, args: string[], optsIn?: ExecOpts) {
-		// We have some special behavior regarding Robomerge being in verbose mode (able to be set through the IPC) --
-		// basically 'debug' is for local development and these messages are really spammy
-		let logLevel : NpmLogLevel =
-			ContextualLogger.getLogLevel() === "verbose" ? "verbose" :
-			"silly"
-			
-
+	static _getP4Cmd(roboWorkspace: RoboWorkspace, args: string[], optsIn?: ExecOpts) {
 		const workspace = coercePerforceWorkspace(roboWorkspace);
 		// add the client explicitly if one is set (should be done at call time)
 
@@ -1637,6 +1580,24 @@ export class PerforceContext {
 
 		args = ['-zprog=robomerge', '-zversion=' + robomergeVersion, ...args]
 
+		return args
+	}
+
+	static _execP4(logger: ContextualLogger, roboWorkspace: RoboWorkspace, args: string[], optsIn?: ExecOpts) {
+		// We have some special behavior regarding Robomerge being in verbose mode (able to be set through the IPC) --
+		// basically 'debug' is for local development and these messages are really spammy
+		let logLevel : NpmLogLevel =
+			ContextualLogger.getLogLevel() === "verbose" ? "verbose" :
+			"silly"
+			
+
+		const workspace = coercePerforceWorkspace(roboWorkspace);
+		// add the client explicitly if one is set (should be done at call time)
+
+		const opts = optsIn || {}
+
+		args = PerforceContext._getP4Cmd(roboWorkspace, args, optsIn)
+
 		// log what we're running
 		let cmd_rec = new CommandRecord('p4 ' + args.join(' '));
 		// 	logger.verbose("Executing: " + cmd_rec.cmd);
@@ -1678,7 +1639,6 @@ export class PerforceContext {
 						fail([new Error(errstr), stderr.toString().replace(newline_rex, '\n')]);
 					}
 					else if (err) {
-						logger.printException(err)
 						let errstr = "P4 Error: " + cmd_rec.cmd + "\n" + err.toString() + "\n";
 
 						if (stdout || stderr) {
@@ -1791,21 +1751,152 @@ export class PerforceContext {
 		return PerforceContext._execP4Ztag(this.logger, roboWorkspace, args, opts)
 	}
 
+	static parseValue(key: string, value: any, parseOptions?: ParseOptions)
+	{
+		if (parseOptions) {
+			const optionalType = parseOptions.optional && parseOptions.optional[key]
+			const fieldType = optionalType || (parseOptions.expected && parseOptions.expected[key]) || 'string'
+			if (fieldType === 'boolean') {
+				if (!optionalType || value) {
+					const valLower = value.toLowerCase()
+					if (valLower !== 'true' && valLower !== 'false') {
+						throw new Error(`Failed to parse boolean field ${key}, value: ${value}`)
+					}
+					return valLower === 'true'
+				}
+				return undefined
+			}
+			else if (fieldType === 'integer') {
+				// ignore empty strings for optional fields (e.g. p4.changes can return a 'shelved' property with no value)
+				if (!optionalType || value) {
+					const num = parseInt(value)
+					if (isNaN(num)) {
+						throw new Error(`Failed to parse number field ${key}, value: ${value}`)
+					}
+					return num
+				}
+				return undefined
+			}
+		}
+		return value		
+	}
+
+	static async execAndParse(logger: ContextualLogger, roboWorkspace: RoboWorkspace, args: string[], execOptions?: ExecOpts, parseOptions?: ParseOptions) {
 
-	////////////
-	// new -ztag parsing - use with care. Initially using for p4.changes and p4.describe
+		args = ['-ztag', '-Mj', ...args]
+		let rawResult = await PerforceContext._execP4(logger, roboWorkspace, args, execOptions)
 
-	async execAndParse(roboWorkspace: RoboWorkspace, args: string[], opts: ExecOpts, options?: ztag.ParseOptions) {
-		const workspace = coercePerforceWorkspace(roboWorkspace);
-		const rawOutput = await PerforceContext._execP4(this.logger, workspace, ['-ztag', ...args], opts)
-		return ztag.parseZtagOutput(rawOutput, this.logger, options)
+		let result = []
+		let startIndex = 0;
+
+		let reviver = (key: string, value: any) => {
+			return PerforceContext.parseValue(key, value, parseOptions)
+		}
+
+		while(startIndex < rawResult.length) {
+			const endIndex = rawResult.indexOf('}\n', startIndex)
+			const parsedResult = JSON.parse(rawResult.slice(startIndex, endIndex != -1 ? endIndex+1 : undefined), reviver)
+			for (const expected in ((parseOptions && parseOptions.expected) || []))
+			{
+				if (!parsedResult[expected])
+				{
+					throw new Error(`Expected field ${expected} not present in ${JSON.stringify(parsedResult)}`)
+				}
+			}
+			result.push(parsedResult)
+			startIndex = endIndex + 2
+		}
 
+		if (result.length == 1 && Object.hasOwn(result[0],'data') && Object.hasOwn(result[0],'generic') && Object.hasOwn(result[0],'severity')) {
+			const cmd = `p4 ${PerforceContext._getP4Cmd(roboWorkspace, args, execOptions).join(' ')}`
+			throw [new Error(`P4 Error: ${cmd}\n${result[0]['data']}`), result[0]['data'].replace(newline_rex, '\n')]
+		}
+
+		return result
 	}
 
-	async execAndParseArray(roboWorkspace: RoboWorkspace, args: string[], opts: ExecOpts, headerOptions?: ztag.ParseOptions, arrayEntryOptions?: ztag.ParseOptions) {
-		const workspace = coercePerforceWorkspace(roboWorkspace);
-		const rawOutput = await PerforceContext._execP4(this.logger, workspace, ['-ztag', ...args], opts)
-		return ztag.parseHeaderAndArray(rawOutput, this.logger, headerOptions, arrayEntryOptions)
+	async execAndParse(roboWorkspace: RoboWorkspace, args: string[], execOptions?: ExecOpts, parseOptions?: ParseOptions) {
+		return PerforceContext.execAndParse(this.logger, roboWorkspace, args, execOptions, parseOptions)
+	}
+
+	static async execAndParseArray(logger: ContextualLogger, roboWorkspace: RoboWorkspace, args: string[], execOptions?: ExecOpts, headerOptions?: ParseOptions, arrayEntryOptions?: ParseOptions) {
+
+		args = ['-ztag', '-Mj', ...args]
+		let rawResult = await PerforceContext._execP4(logger, roboWorkspace, args, execOptions)
+
+		let result = []
+		let startIndex = 0;
+
+		let reviver = (key: string, value: any) => {
+			const arrayElementMatch = key.match(/^(.*?)\d+$/)
+			if (arrayElementMatch) {
+				return this.parseValue(arrayElementMatch[1], value, arrayEntryOptions)
+			}
+			return this.parseValue(key, value, headerOptions)
+		}
+
+		while(startIndex < rawResult.length) {
+			const endIndex = rawResult.indexOf('}\n', startIndex)
+			let parsedResult = JSON.parse(rawResult.slice(startIndex, endIndex != -1 ? endIndex+1 : undefined), reviver)
+			for (let expected in headerOptions && headerOptions.expected || []) {
+				if (!parsedResult[expected])
+				{
+					throw new Error(`Expected field ${expected} not present in ${JSON.stringify(parsedResult)}`)
+				}
+			}
+			let organizedResult: {[key:string]:any} = {};
+			organizedResult.entries = []
+			let expectedCounts: number[] = []
+			for (const field in parsedResult) {
+				const arrayElementMatch = field.match(/^(.*?)(\d+)$/)
+				if (arrayElementMatch) {
+					const arrayField = arrayElementMatch[1]
+					const arrayIndex = parseInt(arrayElementMatch[2])
+					const curLength = expectedCounts.length
+					if (arrayIndex >= curLength) {
+						organizedResult.entries.length = arrayIndex+1
+						expectedCounts.length = arrayIndex+1
+						for (let i=curLength; i<organizedResult.entries.length; i++) {
+							let newEntry: {[key:string]:any} = {};
+							organizedResult.entries[i] = newEntry
+						}
+					}
+					organizedResult.entries[arrayIndex][arrayField] = parsedResult[field]
+					if (arrayEntryOptions && arrayEntryOptions.expected && arrayEntryOptions.expected[arrayField]) {
+						expectedCounts[arrayIndex] += 1
+					}
+				}
+				else {
+					organizedResult[field] = parsedResult[field]
+				}
+			}
+			if (arrayEntryOptions && arrayEntryOptions.expected) {
+				const expectedCount = Object.keys(arrayEntryOptions.expected).length
+				for (let index in expectedCounts) {
+					if (expectedCounts[index] != expectedCount) {
+						for (let expected in arrayEntryOptions.expected) {
+							if (!organizedResult.entries[index][expected])
+							{
+								throw new Error(`Expected field ${expected} not present in ${JSON.stringify(organizedResult.entries[index])}`)
+							}
+						}
+					}
+				}
+			}
+			result.push(organizedResult)
+			startIndex = endIndex + 2
+		}
+
+		if (result.length == 1 && Object.hasOwn(result[0],'data') && Object.hasOwn(result[0],'generic') && Object.hasOwn(result[0],'severity')) {
+			const cmd = `p4 ${PerforceContext._getP4Cmd(roboWorkspace, args, execOptions).join(' ')}`
+			throw [new Error(`P4 Error: ${cmd}\n${result[0]['data']}`), result[0]['data'].replace(newline_rex, '\n')]
+		}
+
+		return result
+	}
+
+	async execAndParseArray(roboWorkspace: RoboWorkspace, args: string[], execOptions?: ExecOpts, headerOptions?: ParseOptions, arrayEntryOptions?: ParseOptions) {
+		return PerforceContext.execAndParseArray(this.logger, roboWorkspace, args, execOptions, headerOptions, arrayEntryOptions)
 	}
 }
 
diff --git a/Engine/Extras/RoboMerge/v3/src/common/ztag.ts b/Engine/Extras/RoboMerge/v3/src/common/ztag.ts
deleted file mode 100644
diff --git a/Engine/Extras/RoboMerge/v3/src/robo/unit-tests.ts b/Engine/Extras/RoboMerge/v3/src/robo/unit-tests.ts
--- a/Engine/Extras/RoboMerge/v3/src/robo/unit-tests.ts
+++ b/Engine/Extras/RoboMerge/v3/src/robo/unit-tests.ts
@@ -5,7 +5,6 @@ import { runTests as gate_runTests } from './gate';
 import { runTests as graph_runTests } from '../new/graph';
 import { runTests as notifications_runTests } from './notifications';
 import { runTests as targets_runTests } from './targets';
-import { runTests as ztag_runTests } from '../common/ztag';
 
 const unitTestsLogger = new ContextualLogger('Unit Tests')
 
@@ -13,8 +12,7 @@ const tests = [
 	gate_runTests,
 	graph_runTests,
 	notifications_runTests,
-	targets_runTests,
-	ztag_runTests
+	targets_runTests
 ]
 
 async function run() {
