Override 2024.1 rules about stream view only integrations so that branch and source/target integrations still work

diff --git a/Engine/Extras/RoboMerge/v3/src/common/perforce.ts b/Engine/Extras/RoboMerge/v3/src/common/perforce.ts
--- a/Engine/Extras/RoboMerge/v3/src/common/perforce.ts
+++ b/Engine/Extras/RoboMerge/v3/src/common/perforce.ts
@@ -447,6 +447,9 @@ export function getPerforceUsername() {
 }
 
 let perforceMultiServerEnvironment: boolean
+let perforceServerVersion: number
+
+const MIN_SERVER_VERSION: number = 2020.1
 
 /**
  * This method must succeed before PerforceContext can be used. Otherwise retrieving the Perforce username through
@@ -469,6 +472,16 @@ export async function initializePerforce(logger: ContextualLogger) {
 	if (resp && resp.User) {
 		perforceUsername = resp.User;
 
+		const serverVersion = await PerforceContext._execP4(logger, null, ["-ztag","-F","%serverVersion%","info"])
+		const match = serverVersion.match(/.*\/(\d+\.\d+)\/.*/)
+		if (!match) {
+			throw new Error(`Unable to parse server version from ${serverVersion}`)
+		}
+		perforceServerVersion = parseFloat(match[1])
+		if (perforceServerVersion < MIN_SERVER_VERSION) {
+			throw new Error(`Robomerge requires a minimum server version of ${MIN_SERVER_VERSION}`)
+		}
+
 		const serversOutput = await PerforceContext.execAndParse(logger, null, ["servers"]);
 		perforceMultiServerEnvironment = serversOutput.length > 1
 	}
@@ -535,7 +548,7 @@ export class PerforceContext {
 
 		const result = await this.execAndParse(workspace, args, {quiet: true, trace: true}, changeResultExpectedShape)
 		if (!result || result.length !== 1) {
-			throw new Error("Expected exactly one change")
+			throw new Error(`Expected exactly one change. Got ${result ? result.length : 0}${result ? '' : '\n' + JSON.stringify(result)}`)
 		}
 
 		const durationSeconds = (Date.now() - startTime) / 1000;
@@ -757,16 +770,6 @@ export class PerforceContext {
 		return null
 	}
 
-	async syncAndReturnChangelistNumber(roboWorkspace: RoboWorkspace, depotPath: string, opts?: string[]) {
-		const workspace = coercePerforceWorkspace(roboWorkspace);
-		const change = await this.latestChange(depotPath);
-		if (!change) {
-			throw new Error('Unable to find changelist');
-		}
-		await this.sync(workspace, `${depotPath}@${change.change}`, {opts});
-		return change.change;
-	}
-
 	async newWorkspace(workspaceName: string, params: any, edgeServer?: {id: string, address: string}) {
 		params.Client = workspaceName
 		if (!('Root' in params)) {
@@ -914,6 +917,10 @@ export class PerforceContext {
 		let noSuchFilesPossible = false // Helper variable for error catching
 		// Branchspec -- takes priority above other possibilities
 		if (source.branchspec) {
+			if (perforceServerVersion >= 2024.1) {
+				// 2024.1 now prevents merges between streams by default
+				cmdList.push("-F")
+			}
 			cmdList.push("-b");
 			cmdList.push(source.branchspec.name);
 			if (source.branchspec.reverse) {
@@ -933,6 +940,10 @@ export class PerforceContext {
 		}
 		// Basic branch to branch
 		else {
+			if (perforceServerVersion >= 2024.1) {
+				// 2024.1 now prevents merges between streams by default
+				cmdList.push("-F")
+			}
 			cmdList.push(source.path_from + range);
 			cmdList.push(target.path_to);
 		}
