Workspace names now take in to account the source as well as the target to avoid multiple edges using the same workspace simultaneously to write to the same stream
No longer precreate all workspaces that might be used, instead create on demand as needed
Auto branch updater will properly handle unloaded workspaces

diff --git a/Engine/Extras/RoboMerge/v3/src/common/perforce.ts b/Engine/Extras/RoboMerge/v3/src/common/perforce.ts
--- a/Engine/Extras/RoboMerge/v3/src/common/perforce.ts
+++ b/Engine/Extras/RoboMerge/v3/src/common/perforce.ts
@@ -91,6 +91,11 @@ export interface DescribeResult {
 	date: Date | null
 }
 
+export interface EdgeServer {
+	id: string, 
+	address: string
+}
+
 /**
 	Example tracing output
 
@@ -681,13 +686,6 @@ export class PerforceContext {
 				throw reason
 			}
 
-			let [err, output] = reason
-
-			// if this change has already been integrated, this is a special return (still a success)
-			if (!output.includes("Revision chars (@, #) not allowed in")) {
-				throw err
-			}
-
 			const errorMsg = `Attempted to find workspaces for invalid user ${user || this.username}`
 			this.logger.error(errorMsg)
 			postToRobomergeAlerts(errorMsg)
@@ -696,11 +694,37 @@ export class PerforceContext {
 		return workspaces as ClientSpec[];
 	}
 
-	find_workspace_by_name(workspaceName: string) {
-		return this.execAndParse(null, ['clients', '-E', workspaceName]);
+	async find_workspace_by_name(workspaceName: string, options?: {edgeServerAddress?: string, includeUnloaded?: boolean}) {
+
+		const edgeServerAddress = options && options.edgeServerAddress
+		const includeUnloaded = options && options.includeUnloaded
+
+		let args = ['clients', '-E', workspaceName]
+
+		// -a to include workspaces on edge servers
+		if (!edgeServerAddress) {
+			// find all
+			args.push('-a')
+		}
+
+		let opts: ExecOpts = {}
+		if (edgeServerAddress && edgeServerAddress !== 'commit') {
+			opts.edgeServerAddress = edgeServerAddress
+		}
+
+		if (edgeServerAddress) {
+			args.push('-a')
+		}
+
+		let result = await this.execAndParse(null, args, opts);
+		if (includeUnloaded && result.length == 0) {
+			result = await this.execAndParse(null, [...args, '-U'], opts)
+		}
+
+		return result
 	}
 
-	async getWorkspaceEdgeServer(workspaceName: string): Promise<{id: string, address: string} | null> {
+	async getWorkspaceEdgeServer(workspaceName: string): Promise<EdgeServer | null> {
 		const serverIdLine = await this._execP4(null, ['-ztag', '-F', '%ServerID%', 'client', '-o', workspaceName]) 
 		if (serverIdLine) {
 			const serverId = serverIdLine.trim()
@@ -770,7 +794,7 @@ export class PerforceContext {
 		return null
 	}
 
-	async newWorkspace(workspaceName: string, params: any, edgeServer?: {id: string, address: string}) {
+	async newWorkspace(workspaceName: string, params: any, edgeServer?: EdgeServer) {
 		params.Client = workspaceName
 		if (!('Root' in params)) {
 			params.Root = 'd:/ROBO/' + workspaceName // default windows path
@@ -817,7 +841,7 @@ export class PerforceContext {
 	}
 
 	// Create a new workspace for Robomerge GraphBot
-	async newGraphBotWorkspace(name: string, extraParams: any, edgeServer?: {id: string, address: string}) {
+	async newGraphBotWorkspace(name: string, extraParams: any, edgeServer?: EdgeServer) {
 		return this.newWorkspace(name, {Root: getRootDirectoryForBranch(name), ...extraParams}, edgeServer);
 	}
 
diff --git a/Engine/Extras/RoboMerge/v3/src/robo/autobranchupdater.ts b/Engine/Extras/RoboMerge/v3/src/robo/autobranchupdater.ts
--- a/Engine/Extras/RoboMerge/v3/src/robo/autobranchupdater.ts
+++ b/Engine/Extras/RoboMerge/v3/src/robo/autobranchupdater.ts
@@ -221,13 +221,16 @@ export class AutoBranchUpdater implements Bot {
 
 		const stream = this.graphBot.branchGraph.config.mirrorPath[0]
 
-		const workspaceQueryResult = await this.p4.find_workspace_by_name(workspace.name)
+		const workspaceQueryResult = await this.p4.find_workspace_by_name(workspace.name, {includeUnloaded: true})
 		if (workspaceQueryResult.length === 0) {
 			await this.p4.newWorkspace(workspace.name, {
 				Stream: stream,
 				Root: AutoBranchUpdater.config!.workspace.directory
 			})
 		}
+		else if (workspaceQueryResult[0].IsUnloaded) {
+			await this.p4.reloadWorkspace(workspace.name)
+		}
 
 		const {depotpath, realFilepath, mirrorFilepath} = workspace
 
diff --git a/Engine/Extras/RoboMerge/v3/src/robo/branch-interfaces.ts b/Engine/Extras/RoboMerge/v3/src/robo/branch-interfaces.ts
--- a/Engine/Extras/RoboMerge/v3/src/robo/branch-interfaces.ts
+++ b/Engine/Extras/RoboMerge/v3/src/robo/branch-interfaces.ts
@@ -1,6 +1,6 @@
 // Copyright Epic Games, Inc. All Rights Reserved.
 
-import { BranchSpec, ConflictedResolveNFile, RoboWorkspace } from '../common/perforce';
+import { BranchSpec, ConflictedResolveNFile, EdgeServer } from '../common/perforce';
 // TODO: Remove Circular Dependency on bot-interfaces
 import { NodeBotInterface } from './bot-interfaces';
 import { FailureKind } from './status-types';
@@ -96,7 +96,6 @@ export interface BranchGraphInterface {
 export type EditableBranch = BranchBase & {
 	bot?: NodeBotInterface
 	parent: BranchGraphInterface
-	workspace: RoboWorkspace
 
 	branchspec: Map<string, BranchSpec>
 	upperName: string
@@ -164,11 +163,8 @@ export interface ChangeInfo extends TargetInfo {
 
 	propagatingNullMerge: boolean
 	forceCreateAShelf: boolean
-	edgeServerToHostShelf?: {
-		id: string
-		address: string
-	}
-	targetWorkspaceOverride?: string
+	edgeServerToHostShelf?: EdgeServer
+	targetWorkspace?: string
 	overriddenCommand: string
 	macros: string[]
 }
diff --git a/Engine/Extras/RoboMerge/v3/src/robo/branchdefs.ts b/Engine/Extras/RoboMerge/v3/src/robo/branchdefs.ts
--- a/Engine/Extras/RoboMerge/v3/src/robo/branchdefs.ts
+++ b/Engine/Extras/RoboMerge/v3/src/robo/branchdefs.ts
@@ -118,7 +118,6 @@ const nodeOptionFieldsPrototype = {
 	streamName: '',
 	streamSubpath: '',
 	uniqueBranch: false,
-	workspace: '',
 
 	graphNodeColor: '',
 
diff --git a/Engine/Extras/RoboMerge/v3/src/robo/branchgraph.ts b/Engine/Extras/RoboMerge/v3/src/robo/branchgraph.ts
--- a/Engine/Extras/RoboMerge/v3/src/robo/branchgraph.ts
+++ b/Engine/Extras/RoboMerge/v3/src/robo/branchgraph.ts
@@ -220,7 +220,6 @@ if (botname === '__TEST__') {
 		let branch: EditableBranch = {
 			name: name,
 			parent: this,
-			workspace: options.workspace || null,
 			branchspec: new Map<string, BranchSpec>(),
 			edgeProperties: new Map(),
 			upperName: nameUpper,
diff --git a/Engine/Extras/RoboMerge/v3/src/robo/edgebot.ts b/Engine/Extras/RoboMerge/v3/src/robo/edgebot.ts
--- a/Engine/Extras/RoboMerge/v3/src/robo/edgebot.ts
+++ b/Engine/Extras/RoboMerge/v3/src/robo/edgebot.ts
@@ -5,7 +5,7 @@ import * as p4util from '../common/p4util';
 
 import { ContextualLogger } from "../common/logger";
 import { Recipients } from "../common/mailer";
-import { Change, coercePerforceWorkspace, ConflictedResolveNFile, EditChangeOpts, EXCLUSIVE_CHECKOUT_REGEX, getRootDirectoryForBranch, IntegrationSource, IntegrationTarget, isExecP4Error, OpenedFileRecord, PerforceContext } from "../common/perforce";
+import { Change, coercePerforceWorkspace, ConflictedResolveNFile, EdgeServer, EditChangeOpts, EXCLUSIVE_CHECKOUT_REGEX, getPerforceUsername, getRootDirectoryForBranch, IntegrationSource, IntegrationTarget, isExecP4Error, OpenedFileRecord, PerforceContext } from "../common/perforce";
 import { VersionReader } from "../common/version";
 import { EdgeBotInterface, IPCControls, ReconsiderArgs } from "./bot-interfaces";
 import { AlreadyIntegrated, Branch, ChangeInfo, ConflictingFile, ExclusiveFile, Failure, MergeAction, PendingChange } from "./branch-interfaces";
@@ -298,7 +298,7 @@ class EdgeBotImpl extends PerforceStatefulBot {
 		}
 
 		this.edgeBotLogger.info(`${logMessage}. Reverting ${pending.newCl}.`)
-		await this.revertAndDelete(coercePerforceWorkspace(this.targetBranch.workspace)!.name, pending.newCl)
+		await this.revertAndDelete(coercePerforceWorkspace(pending.change.targetWorkspace)!.name, pending.newCl)
 
 		let pauseDurationSeconds = FAILED_CHANGELIST_PAUSE_TIMEOUT_SECONDS
 		// if we have a target, make sure pause duration is at least 2x duration of failed integration
@@ -454,52 +454,86 @@ class EdgeBotImpl extends PerforceStatefulBot {
 		this.sourceNode.queueEdgeUnblock(this.targetBranch.upperName)
 	}
 
-	private async integrate(info: ChangeInfo, target: MergeAction) : Promise<EdgeIntegrationDetails> {
-		const to_integrate = info.cl
-		this._log_action(`Integrating CL ${to_integrate} to ${this.targetBranch.name}`)
+	public async getWorkspace(edgeServer?: EdgeServer) {
 
-		// if required, add author review here so they're not in target.description, which is used for shelf description in case of conflict
-		const desc = target.description! // target.description always ends in newline
+		const edgeServerAddress: string | undefined = edgeServer && edgeServer.address
 
-		info.targetWorkspaceOverride = coercePerforceWorkspace(this.targetBranch.workspace)!.name
-		const edgeServer = info.edgeServerToHostShelf
+		// name the workspace
+		let workspaceName = this.targetBranch.config.workspaceNameOverride || ['ROBOMERGE', this.targetBranch.parent.botname, this.targetBranch.name].join('_');
+		let fromName = `_FROM_${(this.targetBranch.parent == this.sourceBranch.parent ? '' : this.sourceBranch.parent.botname)}_${this.sourceBranch.name}`
+		if (this.incognitoMode) {
+			const hashCode = (s: string) => s.split('').reduce((a,b)=>{a=((a<<5)-a)+b.charCodeAt(0);return a&a},0)
+			workspaceName += '_' + hashCode(fromName)
+		}
+		else {
+			workspaceName += fromName
+		}
+		const p4username = getPerforceUsername()
+		if (p4username !== 'robomerge') {
+			workspaceName = [p4username!.toUpperCase(), process.platform.toUpperCase(), workspaceName].join('_')
+		}
+		workspaceName = workspaceName.replace(/[\/\.-\s]/g, "_").replace(/_+/g,"_");
 		if (edgeServer) {
+			workspaceName += '_' + edgeServer.id.toUpperCase()
+		}
 
-			info.targetWorkspaceOverride += '_' + edgeServer.id.toUpperCase()
-
-			// make sure target workspace/directory exists (always reset for now)
+		// ensure root directory exists (we set the root diretory to be the cwd)
+		const path = getRootDirectoryForBranch(workspaceName);
+		if (!fs.existsSync(path)) {
+			this.logger.info(`Making directory ${path}`);
+			fs.mkdirSync(path);
+		}
 
-			// ensure local directory exists (_initWorkspacesForGraphBot)
-			const path = getRootDirectoryForBranch(info.targetWorkspaceOverride);
-			if (!fs.existsSync(path)) {
-				this.edgeBotLogger.info(`Making directory ${path}`);
-				fs.mkdirSync(path);
+		// do we already have a workspace?
+		const existingWorkspaceInfo = await this.p4.find_workspace_by_name(workspaceName, {edgeServerAddress, includeUnloaded: true})
+		if (existingWorkspaceInfo.length > 0) {
+			if (existingWorkspaceInfo[0].IsUnloaded) {
+				await this.p4.reloadWorkspace(workspaceName, edgeServerAddress)
 			}
-
-			// do we already have a workspace? (_initBranchWorkspacesForAllBots, _getExistingWorkspaces)
-			const existingWorkspaceInfos = await this.p4.find_workspaces(undefined, {edgeServerAddress: edgeServer.address, includeUnloaded: true})
-			const existingWorkspaceIndex = existingWorkspaceInfos.findIndex((ws) => ws.client == info.targetWorkspaceOverride)
-			if (existingWorkspaceIndex >= 0) {
-				if (existingWorkspaceInfos[existingWorkspaceIndex].IsUnloaded) {
-					await this.p4.reloadWorkspace(info.targetWorkspaceOverride, edgeServer.address)
-				}
-				await p4util.cleanWorkspaces(this.p4, [[info.targetWorkspaceOverride, target.branch.rootPath]], edgeServer.address)
+			await p4util.cleanWorkspaces(this.p4, [[workspaceName, this.targetBranch.rootPath]], edgeServerAddress)
+		}
+		else {
+			const params: any = {};
+			if (this.targetBranch.stream) {
+				params['Stream'] = this.targetBranch.stream;
 			}
 			else {
-				// create one 
-				if (!target.branch.stream) {
-					throw new Error('only stream workspaces supported on edge servers')
-				}
-				await this.p4.newGraphBotWorkspace(info.targetWorkspaceOverride, {Stream: target.branch.stream}, edgeServer)
+				params['View'] = [
+					`${this.targetBranch.rootPath} //${workspaceName}/...`
+				];
+			}
+
+			await this.p4.newGraphBotWorkspace(workspaceName, params, edgeServer);
+
+			// if we're on linux, remove the directory whenever we create the workspace for the first time
+			if (process.platform === "linux") {
+				const dir = '/src/' + workspaceName;
+				this.logger.info(`Cleaning ${dir}...`);
 
-				// _initWorkspacesForGraphBot does clean-up stuff here, but I don't think it's necessary
+				// delete the directory contents (but not the directory)
+				require('child_process').execSync(`rm -rf ${dir}/*`);
+			}
+			else {
+				await this.p4.clean(workspaceName);
 			}
 		}
+		
+		return workspaceName
+	}
 
-		const edgeServerAddress: string | undefined = edgeServer && edgeServer.address
+	private async integrate(info: ChangeInfo, target: MergeAction) : Promise<EdgeIntegrationDetails> {
+		const to_integrate = info.cl
+		this._log_action(`Integrating CL ${to_integrate} to ${this.targetBranch.name}`)
+
+		// if required, add author review here so they're not in target.description, which is used for shelf description in case of conflict
+		const desc = target.description! // target.description always ends in newline
+
+		info.targetWorkspace = await this.getWorkspace(info.edgeServerToHostShelf)
+
+		const edgeServerAddress = info.edgeServerToHostShelf && info.edgeServerToHostShelf.address
 
 		// create a new CL
-		const changenum = await this.p4.new_cl(info.targetWorkspaceOverride, desc, undefined, edgeServerAddress)
+		const changenum = await this.p4.new_cl(info.targetWorkspace, desc, undefined, edgeServerAddress)
 
 		// try to integrate
 		const branchSpecToTarget = this.sourceBranch.branchspec.get(this.targetBranch.upperName)
@@ -524,8 +558,8 @@ class EdgeBotImpl extends PerforceStatefulBot {
 		    || (!info.userRequest && !info.forceCreateAShelf && !target.flags.has('manual'))
 
 		this.currentIntegrationStartTimestamp = Date.now()
-		await this.p4.sync(info.targetWorkspaceOverride, this.targetBranch.rootPath + '#0', {edgeServerAddress})
-		const [mode, results] = await this.p4.integrate(info.targetWorkspaceOverride, source, changenum, integTarget, {edgeServerAddress, virtual: doVirtualMerge})
+		await this.p4.sync(info.targetWorkspace, this.targetBranch.rootPath + '#0', {edgeServerAddress})
+		const [mode, results] = await this.p4.integrate(info.targetWorkspace, source, changenum, integTarget, {edgeServerAddress, virtual: doVirtualMerge})
 
 		const pending: PendingChange = {change: info, action: target, newCl: changenum}
 
@@ -543,7 +577,7 @@ class EdgeBotImpl extends PerforceStatefulBot {
 
 				// integration not necessary
 				this.edgeBotLogger.info(msg)
-				await this.p4.deleteCl(info.targetWorkspaceOverride, changenum, edgeServerAddress)
+				await this.p4.deleteCl(info.targetWorkspace, changenum, edgeServerAddress)
 				return new EdgeIntegrationDetails('nothing to do', msg)
 			}
 		}
@@ -565,7 +599,7 @@ class EdgeBotImpl extends PerforceStatefulBot {
 
 		// Revert attempt
 		if (pending.newCl > 0) {
-			await this.revertAndDelete(info.targetWorkspaceOverride, pending.newCl, edgeServerAddress)
+			await this.revertAndDelete(info.targetWorkspace, pending.newCl, edgeServerAddress)
 			pending.newCl = -1
 		}
 		
@@ -612,7 +646,7 @@ class EdgeBotImpl extends PerforceStatefulBot {
 //		const detail: ResolveResultDetail = 'detailed'
 
 		const result = await this.p4.resolve(
-			pending.change.targetWorkspaceOverride || this.targetBranch.workspace,
+			pending.change.targetWorkspace,
 			pending.newCl,
 			pending.action.mergeMode,
 			false, // detail === 'quick',
@@ -709,7 +743,7 @@ class EdgeBotImpl extends PerforceStatefulBot {
 		// try to submit
 		this._log_action(`Submitting CL ${changenum} by ${info.author}`)
 		const result = await this.p4.submit(
-			pending.change.targetWorkspaceOverride || this.targetBranch.workspace, 
+			pending.change.targetWorkspace, 
 			changenum,
 			pending.change.edgeServerToHostShelf && pending.change.edgeServerToHostShelf.address)
 			
@@ -744,7 +778,7 @@ class EdgeBotImpl extends PerforceStatefulBot {
 			// change owner, so users can edit change descriptions later for reconsideration
 			this.edgeBotLogger.info(`Setting owner of CL ${finalCl} to author of change: ${info.author}`);
 			try {
-				await this.p4.editOwner(this.targetBranch.workspace, finalCl, info.author, {changeSubmitted: true})
+				await this.p4.editOwner(info.targetWorkspace, finalCl, info.author, {changeSubmitted: true})
 			}
 			catch (reason) {
 				let errPreface = 'Error changing owner'
@@ -798,12 +832,8 @@ class EdgeBotImpl extends PerforceStatefulBot {
 			final_desc += pending.change.additionalDescriptionText
 		}
 
-		/*if (forApproval) {
-			final_desc += "\n#review"
-		}*/
-
 		const edgeServerAddress = pending.change.edgeServerToHostShelf && pending.change.edgeServerToHostShelf.address
-		const destRoboWorkspace = pending.change.targetWorkspaceOverride || this.targetBranch.workspace
+		const destRoboWorkspace = pending.change.targetWorkspace
 
 		// abort shelve if this is a buildmachine / robomerge change (unless we are forcing the shelf)
 		let failed = false
@@ -1103,6 +1133,10 @@ export class EdgeBot
 		this.impl.queueUnblock()
 	}
 
+	getWorkspace(edgeServer?: EdgeServer) {
+		return this.impl.getWorkspace(edgeServer)
+	}
+
 	/* Mirrored Variables */
 	get disallowSkip() { return this.impl.disallowSkip }
 	get incognitoMode() { return this.impl.incognitoMode }
diff --git a/Engine/Extras/RoboMerge/v3/src/robo/graphbot.ts b/Engine/Extras/RoboMerge/v3/src/robo/graphbot.ts
--- a/Engine/Extras/RoboMerge/v3/src/robo/graphbot.ts
+++ b/Engine/Extras/RoboMerge/v3/src/robo/graphbot.ts
@@ -3,8 +3,7 @@
 import * as Sentry from '@sentry/node';
 import { ContextualLogger } from '../common/logger';
 import { Mailer } from '../common/mailer';
-import * as p4util from '../common/p4util';
-import { PerforceContext, Workspace } from '../common/perforce';
+import { PerforceContext } from '../common/perforce';
 import { AutoBranchUpdater } from './autobranchupdater';
 import { bindBadgeHandler } from './badges';
 import { Bot } from './bot-interfaces';
@@ -185,20 +184,6 @@ export class GraphBot implements GraphInterface, BotEventHandler {
 		this.botLogger.info(msg)
 		postToRobomergeAlerts(msg)
 
-		if (this.branchGraph.branches.length !== 0) {
-			const workspaces = this.branchGraph.branches.map(branch => 
-				[(branch.workspace as Workspace).name || (branch.workspace as string),
-				branch.rootPath]) as [string, string][]
-			const mirrorWorkspace = AutoBranchUpdater.getMirrorWorkspace(this)
-			if (mirrorWorkspace) {
-				// add /... to match branches' rootPath format
-				workspaces.push([mirrorWorkspace.name, mirrorWorkspace.stream + '/...'])
-			}
-
-			this.botLogger.info('Cleaning all workspaces')
-			await p4util.cleanWorkspaces(this.p4, workspaces)
-		}
-
 		await this.startBotsAsync()
 	}
 
diff --git a/Engine/Extras/RoboMerge/v3/src/robo/nodebot.ts b/Engine/Extras/RoboMerge/v3/src/robo/nodebot.ts
--- a/Engine/Extras/RoboMerge/v3/src/robo/nodebot.ts
+++ b/Engine/Extras/RoboMerge/v3/src/robo/nodebot.ts
@@ -128,12 +128,6 @@ export class NodeBot extends PerforceStatefulBot implements NodeBotInterface {
 
 		this.initTickJournal()
 
-		if (!this.branch.workspace) {
-			throw new Error(`Branch ${this.fullName} has no valid workspace specified`)
-		}
-
-		// not looking for min CL of edges any more
-
 		// Finally after setup, create the edges. (Edges may rely on NodeBot data for setup, so always do this last.)
 		this.edges = this.createEdges()
 	}
@@ -667,7 +661,8 @@ export class NodeBot extends PerforceStatefulBot implements NodeBotInterface {
 		const edgeMap = new Map<string, EdgeBot>([[edge.targetBranch.upperName, edge]])
 		let integratedCl: Change | null = null // Set to possible undefines as we use this in the finally clause
 		try {
-			const processResult = await this._createChangeInfo(blockageChange, edgeMap, this.p4.username, targetBranch.workspace, targetBranch)
+			const workspace = await edge.getWorkspace()
+			const processResult = await this._createChangeInfo(blockageChange, edgeMap, this.p4.username, workspace, targetBranch)
 			if (!processResult.info) {
 				return {
 					success: false, 
@@ -711,7 +706,7 @@ export class NodeBot extends PerforceStatefulBot implements NodeBotInterface {
 			integratedCl = p4ChangeResult as Change
 
 			// Unshelve 
-			if (!await this.p4.unshelve(targetBranch.workspace, integratedCl.change)) {
+			if (!await this.p4.unshelve(workspace, integratedCl.change)) {
 				return { 
 					success: false, 
 					message: `Unable to unshelve CL ${integratedCl.change}`
diff --git a/Engine/Extras/RoboMerge/v3/src/robo/robo.ts b/Engine/Extras/RoboMerge/v3/src/robo/robo.ts
--- a/Engine/Extras/RoboMerge/v3/src/robo/robo.ts
+++ b/Engine/Extras/RoboMerge/v3/src/robo/robo.ts
@@ -5,13 +5,12 @@ import { OnUncaughtException, OnUnhandledRejection } from '@sentry/node/dist/int
 import * as fs from 'fs';
 import * as os from 'os';
 import * as path from 'path';
-import * as p4util from '../common/p4util';
 import { Analytics } from '../common/analytics';
 import { Arg, readProcessArgs } from '../common/args';
 import { _setTimeout } from '../common/helper';
 import { ContextualLogger } from '../common/logger';
 import { Mailer } from '../common/mailer';
-import { ClientSpec, getPerforceUsername, getRootDirectoryForBranch, initializePerforce, PerforceContext } from '../common/perforce';
+import { getPerforceUsername, initializePerforce, PerforceContext } from '../common/perforce';
 import { BuildVersion, VersionReader } from '../common/version';
 import { CertFiles } from '../common/webserver';
 import { addBranchGraph, Graph, GraphAPI } from '../new/graph';
@@ -35,9 +34,6 @@ import * as Preview from './preview'
 const roboStartupLogger = new ContextualLogger('Robo Startup')
 VersionReader.init(roboStartupLogger)
 
-// I seem to have broken this
-const DEBUG_SKIP_BRANCH_SETUP = false;
-
 const COMMAND_LINE_ARGS: {[param: string]: Arg<any>} = {
 	botname: {
 		match: /^-botname=([a-zA-Z0-9_,]+)$/,
@@ -272,88 +268,6 @@ export class RoboMerge {
 	}
 }
 
-async function _getExistingWorkspaces(user?: string) {
-	const existingWorkspacesObjects = await robo.p4.find_workspaces(user, {includeUnloaded: true});
-	return new Map<string, ClientSpec>(existingWorkspacesObjects.map((ws: ClientSpec) => [ws.client, ws]));
-}
-
-async function _initWorkspacesForGraphBot(graphBot: GraphBot, existingWorkspaces: Map<string, ClientSpec>, logger: ContextualLogger) {
-	// name and depot root pair
-	const workspacesToReset: [string, string][] = []
-	let reloadedWorkspaces = []
-
-	for (const branch of graphBot.branchGraph.branches) {
-		if (branch.workspace !== null) {
-			logger.info(`Using manually configured workspace ${branch.workspace} for branch ${branch.name}`);
-			continue;
-		}
-
-		// name the workspace
-		let workspaceName = branch.config.workspaceNameOverride || ['ROBOMERGE', branch.parent.botname, branch.name].join('_');
-		const p4username = getPerforceUsername()
-		if (p4username !== 'robomerge') {
-			workspaceName = [p4username!.toUpperCase(), process.platform.toUpperCase(), workspaceName].join('_')
-		}
-		branch.workspace = workspaceName.replace(/[\/\.-\s]/g, "_").replace(/_+/g,"_");
-
-		const ws = branch.workspace;
-
-		// ensure root directory exists (we set the root diretory to be the cwd)
-		const path = getRootDirectoryForBranch(ws);
-		if (!fs.existsSync(path)) {
-			logger.info(`Making directory ${path}`);
-			fs.mkdirSync(path);
-		}
-
-		// see if we already have this workspace
-		const clientSpec = existingWorkspaces.get(ws)
-		if (clientSpec && clientSpec.Stream === branch.stream) {
-			if (clientSpec.IsUnloaded) {
-				reloadedWorkspaces.push(robo.p4.reloadWorkspace(ws))
-			}
-			workspacesToReset.push([ws, branch.rootPath])
-		}
-		else {
-			const params: any = {};
-			if (branch.stream) {
-				params['Stream'] = branch.stream;
-			}
-			else {
-				params['View'] = [
-					`${branch.rootPath} //${ws}/...`
-				];
-			}
-
-			await robo.p4.newGraphBotWorkspace(ws, params);
-
-			// if we're on linux, remove the directory whenever we create the workspace for the first time
-			if (process.platform === "linux") {
-				const dir = '/src/' + branch.workspace;
-				logger.info(`Cleaning ${dir}...`);
-
-				// delete the directory contents (but not the directory)
-				require('child_process').execSync(`rm -rf ${dir}/*`);
-			}
-			else {
-				await robo.p4.clean(branch.workspace);
-			}
-		}
-	}
-
-	await Promise.all(reloadedWorkspaces)
-	if (workspacesToReset.length > 0) {
-		logger.info('The following workspaces already exist and will be reset: ' + workspacesToReset.join(', '))
-		await p4util.cleanWorkspaces(robo.p4, workspacesToReset)
-	}
-}
-
-async function _initBranchWorkspacesForAllBots(logger: ContextualLogger) {
-	const existingWorkspaces = await _getExistingWorkspaces();
-	for (const graphBot of robo.graphBots.values()) {
-		await _initWorkspacesForGraphBot(graphBot, existingWorkspaces, logger);
-	}
-}
-
 // should be called after data directory has been synced, to get latest email template
 function _initMailer(logger: ContextualLogger) {
 	robo.mailer = new Mailer(roboAnalytics!, logger, args.noMail);
@@ -387,7 +301,6 @@ let specReloadEntryCount = 0
 async function _onBranchSpecReloaded(graphBot: GraphBot, logger: ContextualLogger) {
 	try {
 		++specReloadEntryCount
-		await _initWorkspacesForGraphBot(graphBot, await _getExistingWorkspaces(), logger)
 	}
 	finally {
 		--specReloadEntryCount
@@ -499,9 +412,6 @@ async function init(logger: ContextualLogger) {
 
 	_initMailer(logger)
 	await _initGraphBots(logger)
-	if (!DEBUG_SKIP_BRANCH_SETUP) {
-		await _initBranchWorkspacesForAllBots(logger)
-	}
 
 	const graph = new Graph
 	robo.graph = new GraphAPI(graph)
