Fix crash when reconsidering massive changelists on a block asset edge

diff --git a/Engine/Extras/RoboMerge/v3/src/common/perforce.ts b/Engine/Extras/RoboMerge/v3/src/common/perforce.ts
--- a/Engine/Extras/RoboMerge/v3/src/common/perforce.ts
+++ b/Engine/Extras/RoboMerge/v3/src/common/perforce.ts
@@ -1445,11 +1445,18 @@ export class PerforceContext {
 		return (await this.execAndParse(null, ['dirs', path])).map((dir: any) => dir.dir)
 	}
 
-	async files(path: string, maxFiles?: number) {
+	files(path: string, maxFiles?: number) {
 		const args = ['files',
 			...(maxFiles ? ['-m', maxFiles.toString()] : []),
 			path]
-		return await this.execAndParse(null, args);
+		return this.execAndParse(null, args);
+	}
+
+	sizes(path: string, summary?: boolean) {
+		const args = ['sizes',
+			...(summary ? ['-s'] : []),
+			path]
+		return this.execAndParse(null, args);
 	}
 
 	// update the fields on an existing CL using p4 change
diff --git a/Engine/Extras/RoboMerge/v3/src/robo/nodebot.ts b/Engine/Extras/RoboMerge/v3/src/robo/nodebot.ts
--- a/Engine/Extras/RoboMerge/v3/src/robo/nodebot.ts
+++ b/Engine/Extras/RoboMerge/v3/src/robo/nodebot.ts
@@ -1681,19 +1681,11 @@ export class NodeBot extends PerforceStatefulBot implements NodeBotInterface {
 			}
 
 			if (blockAssetEdges.length > 0) {
-				const describeResult = await this.p4.describe(change.change)
-
-				let changeContainsAssets = false
-				for (const entry of describeResult!.entries) {
-					const fileExtIndex = entry.depotFile.lastIndexOf('.')
-					if (fileExtIndex !== -1) {
-						const fileExt = entry.depotFile.substring(fileExtIndex + 1)
-						if (fileExt === 'uasset' || fileExt === 'umap') {
-							changeContainsAssets = true
-							break
-						}
-					}
-				}
+
+				let uassetSize = this.p4.sizes(`//....uasset@=${change.change}`, true)
+				let umapSize = this.p4.sizes(`//....umap@=${change.change}`, true)
+
+				const changeContainsAssets = ((await uassetSize)[0].fileCount > 0 || (await umapSize)[0].fileCount > 0)
 
 				if (changeContainsAssets) {
 					for (const [edge, action] of blockAssetEdges) {
