Do not update mirror files in preview only

diff --git a/Engine/Extras/RoboMerge/v3/src/robo/autobranchupdater.ts b/Engine/Extras/RoboMerge/v3/src/robo/autobranchupdater.ts
--- a/Engine/Extras/RoboMerge/v3/src/robo/autobranchupdater.ts
+++ b/Engine/Extras/RoboMerge/v3/src/robo/autobranchupdater.ts
@@ -48,7 +48,7 @@ export class AutoBranchUpdater implements Bot {
 
 	tickCount = 0
 
-	constructor(private graphBot: GraphInterface, parentLogger: ContextualLogger) {
+	constructor(private graphBot: GraphInterface, parentLogger: ContextualLogger, private readonly previewMode: boolean) {
 		this.p4 = AutoBranchUpdater.p4!
 
 		const config = AutoBranchUpdater.config
@@ -217,6 +217,13 @@ export class AutoBranchUpdater implements Bot {
 	}
 
 	private async updateMirror(workspace: MirrorPaths) {
+
+		if (this.previewMode)
+		{
+			this.abuLogger.info("Skipping mirror update in Preview Mode")
+			return
+		}
+
 		this.abuLogger.info("Updating branchmap mirror")
 
 		const stream = this.graphBot.branchGraph.config.mirrorPath[0]
diff --git a/Engine/Extras/RoboMerge/v3/src/robo/robo.ts b/Engine/Extras/RoboMerge/v3/src/robo/robo.ts
--- a/Engine/Extras/RoboMerge/v3/src/robo/robo.ts
+++ b/Engine/Extras/RoboMerge/v3/src/robo/robo.ts
@@ -427,7 +427,7 @@ function startBots(logger: ContextualLogger) {
 	logger.info("Starting branch bots...");
 	for (let graphBot of robo.graphBots.values()) {
 		if (AutoBranchUpdater.config) {
-			graphBot.autoUpdater = new AutoBranchUpdater(graphBot, logger)
+			graphBot.autoUpdater = new AutoBranchUpdater(graphBot, logger, args.previewOnly)
 		}
 		graphBot.runbots()
 	}
