Move workspaceNameOverride from the node to the edge
Reenable workspace including source name logic

diff --git a/Engine/Extras/RoboMerge/v3/README.md b/Engine/Extras/RoboMerge/v3/README.md
--- a/Engine/Extras/RoboMerge/v3/README.md
+++ b/Engine/Extras/RoboMerge/v3/README.md
@@ -98,7 +98,6 @@ for both node and edge.
 |`depot`                    |n |**required**| |Depot, e.g. UE4                          |Flag currently called `defaultStreamDepot`|
 |`rootPath`                 |n |from stream | |P4 depot full path                       | |
 |`streamSubpath`            |n |`/...`      | |P4 depot sub-path                        | |
-|`workspaceNameOverride`    |n |            | |Used specified workspace name for commits| |
 |`enabled`                  |ne|`true`      | |If false, pretends node/edge doesn't exist| |
 |`forcePause`               |ne|            | |If flag set, pause - applies each restart| |
 |`additionalSlackChannelForBlockages` |e |  | |Single extra Slack channel               | |
@@ -113,6 +112,7 @@ for both node and edge.
 |`isDefaultBot`             |e |`false`     | |Run plain #robomerge commands? Should be `false` for streams monitored by multiple bots| |
 |`lastGoodCLPath`           |e |            | |'Gate' file to read to find CIS-approved CL| |
 |`waitingForCISLink`        |e |            | |URL to link to for waiting for CIS messages| |
+|`workspaceNameOverride`    |e |            | |User specified workspace name for commits| |
 |`notify`                   |e |            | |Additional people to email on blockages  |Also `globalNotify`|
 |`resolver`                 |e |            | |Single designated resolver               |Currently applies to both source and target nodes|
 
diff --git a/Engine/Extras/RoboMerge/v3/src/robo/edgebot.ts b/Engine/Extras/RoboMerge/v3/src/robo/edgebot.ts
--- a/Engine/Extras/RoboMerge/v3/src/robo/edgebot.ts
+++ b/Engine/Extras/RoboMerge/v3/src/robo/edgebot.ts
@@ -459,15 +459,15 @@ class EdgeBotImpl extends PerforceStatefulBot {
 		const edgeServerAddress: string | undefined = edgeServer && edgeServer.address
 
 		// name the workspace
-		let workspaceName = this.targetBranch.config.workspaceNameOverride || ['ROBOMERGE', this.targetBranch.parent.botname, this.targetBranch.name].join('_');
-		/*let fromName = `_FROM_${(this.targetBranch.parent == this.sourceBranch.parent ? '' : this.sourceBranch.parent.botname)}_${this.sourceBranch.name}`
+		let workspaceName = this.options.workspaceNameOverride || ['ROBOMERGE', this.targetBranch.parent.botname, this.targetBranch.name].join('_');
+		let fromName = `_FROM_${(this.targetBranch.parent == this.sourceBranch.parent ? '' : this.sourceBranch.parent.botname)}_${this.sourceBranch.name}`
 		if (this.incognitoMode) {
 			const hashCode = (s: string) => s.split('').reduce((a,b)=>{a=((a<<5)-a)+b.charCodeAt(0);return a&a},0)
 			workspaceName += '_' + hashCode(fromName)
 		}
 		else {
 			workspaceName += fromName
-		}*/
+		}
 		const p4username = getPerforceUsername()
 		if (p4username !== 'robomerge') {
 			workspaceName = [p4username!.toUpperCase(), process.platform.toUpperCase(), workspaceName].join('_')
