Fix execAndParse not reporting errors if there was more than 1 result returned
Exclusive Checkout test adapted to do a multi-file changelist with one file locked and one able to be merged

diff --git a/Engine/Extras/RoboMerge/v3/functional_tests/framework/src/tests/exclusive-checkout.ts b/Engine/Extras/RoboMerge/v3/functional_tests/framework/src/tests/exclusive-checkout.ts
--- a/Engine/Extras/RoboMerge/v3/functional_tests/framework/src/tests/exclusive-checkout.ts
+++ b/Engine/Extras/RoboMerge/v3/functional_tests/framework/src/tests/exclusive-checkout.ts
@@ -16,6 +16,7 @@ export class ExclusiveCheckout extends FunctionalTest {
 
 		const mainClient = this.getClient('Main', 'testuser1')
 		await P4Util.addFileAndSubmit(mainClient, 'test.uasset', 'Dummy content', true)
+		await P4Util.addFileAndSubmit(mainClient, 'test2.uasset', 'Dummy content', true)
 		
 		this.releaseUser2Client = this.getClient('Release', 'testuser2')
 		await Promise.all([
@@ -26,8 +27,10 @@ export class ExclusiveCheckout extends FunctionalTest {
 	}
 
 	run() {
-		// second user commits change in Release well first user has file in main checked out
-		return P4Util.editFileAndSubmit(this.releaseUser2Client, 'test.uasset', 'New content')
+		// second user commits change in Release while first user has file in main checked out
+		P4Util.editFile(this.releaseUser2Client, 'test.uasset', 'New content')
+		P4Util.editFile(this.releaseUser2Client, 'test2.uasset', 'New content')
+		return P4Util.submit(this.releaseUser2Client, "Submit by second user")
 	}
 
 	verify() {
diff --git a/Engine/Extras/RoboMerge/v3/src/common/perforce.ts b/Engine/Extras/RoboMerge/v3/src/common/perforce.ts
--- a/Engine/Extras/RoboMerge/v3/src/common/perforce.ts
+++ b/Engine/Extras/RoboMerge/v3/src/common/perforce.ts
@@ -1849,9 +1849,14 @@ export class PerforceContext {
 			startIndex = endIndex + 2
 		}
 
-		if (result.length == 1 && Object.hasOwn(result[0],'data') && Object.hasOwn(result[0],'generic') && Object.hasOwn(result[0],'severity')) {
+		let error = result
+						.filter((r) => Object.hasOwn(r,'data') && Object.hasOwn(r,'generic') && Object.hasOwn(r,'severity'))
+						.map((r) => r.data)
+						.join('')
+
+		if (error.length > 0) {
 			const cmd = `p4 ${PerforceContext._getP4Cmd(roboWorkspace, args, execOptions).join(' ')}`
-			throw [new Error(`P4 Error: ${cmd}\n${result[0]['data']}`), result[0]['data'].replace(newline_rex, '\n')]
+			throw [new Error(`P4 Error: ${cmd}\n${error}`), error.replace(newline_rex, '\n')]
 		}
 
 		return result
