Don't try to run the change owner on an edge server when reconsidering

diff --git a/Engine/Extras/RoboMerge/v3/src/robo/edgebot.ts b/Engine/Extras/RoboMerge/v3/src/robo/edgebot.ts
--- a/Engine/Extras/RoboMerge/v3/src/robo/edgebot.ts
+++ b/Engine/Extras/RoboMerge/v3/src/robo/edgebot.ts
@@ -778,7 +778,7 @@ class EdgeBotImpl extends PerforceStatefulBot {
 			// change owner, so users can edit change descriptions later for reconsideration
 			this.edgeBotLogger.info(`Setting owner of CL ${finalCl} to author of change: ${info.author}`);
 			try {
-				await this.p4.editOwner(info.targetWorkspace, finalCl, info.author, {changeSubmitted: true})
+				await this.p4.editOwner(null, finalCl, info.author, {changeSubmitted: true})
 			}
 			catch (reason) {
 				let errPreface = 'Error changing owner'
