Remove unnecessary coerce workspace calls

diff --git a/Engine/Extras/RoboMerge/v3/src/common/perforce.ts b/Engine/Extras/RoboMerge/v3/src/common/perforce.ts
--- a/Engine/Extras/RoboMerge/v3/src/common/perforce.ts
+++ b/Engine/Extras/RoboMerge/v3/src/common/perforce.ts
@@ -746,10 +746,9 @@ export class PerforceContext {
 	}
 
 	async clean(roboWorkspace: RoboWorkspace) {
-		const workspace = coercePerforceWorkspace(roboWorkspace);
 		let result;
 		try {
-			result = await this._execP4(workspace, ['clean', '-e', '-a']);
+			result = await this._execP4(roboWorkspace, ['clean', '-e', '-a']);
 		}
 		catch (reason) {
 			if (!isExecP4Error(reason)) {
@@ -768,7 +767,6 @@ export class PerforceContext {
 
 	// sync the depot path specified
 	async sync(roboWorkspace: RoboWorkspace, depotPath: string, params?: SyncParams) {
-		const workspace = coercePerforceWorkspace(roboWorkspace);
 		let args = params && params.edgeServerAddress ? ['-p', params.edgeServerAddress] : []
 		args.push('sync')
 		if (params && params.opts) {
@@ -776,7 +774,7 @@ export class PerforceContext {
 		}
 		args.push(depotPath)
 		try {
-			return await this.execAndParse(workspace, args)
+			return await this.execAndParse(roboWorkspace, args)
 		}
 		catch (reason) {
 			if (!isExecP4Error(reason)) {
@@ -923,7 +921,6 @@ export class PerforceContext {
 		: Promise<[string, (Change | string)[]]> {
 
 		opts = opts || {};
-		const workspace = coercePerforceWorkspace(roboWorkspace);
 
 		// build a command
 		let cmdList = [
@@ -975,7 +972,7 @@ export class PerforceContext {
 		// execute the P4 command
 		let changes;
 		try {
-			changes = await this.execAndParse(workspace, cmdList, { numRetries: 0, edgeServerAddress: opts.edgeServerAddress });
+			changes = await this.execAndParse(roboWorkspace, cmdList, { numRetries: 0, edgeServerAddress: opts.edgeServerAddress });
 		}
 		catch (reason) {
 			if (!isExecP4Error(reason)) {
@@ -1138,10 +1135,9 @@ export class PerforceContext {
 	// submit a CL
 	// output format is final CL number or false if changes need more resolution
 	async submit(roboWorkspace: RoboWorkspace, changelist: number, edgeServerAddress?: string): Promise<number | string> {
-		const workspace = coercePerforceWorkspace(roboWorkspace)
 		let rawOutput: string
 		try {
-			rawOutput = await this._execP4(workspace, ['-ztag', 'submit', '-f', 'submitunchanged', '-c', changelist.toString()], {edgeServerAddress})
+			rawOutput = await this._execP4(roboWorkspace, ['-ztag', 'submit', '-f', 'submitunchanged', '-c', changelist.toString()], {edgeServerAddress})
 		}
 		catch ([errArg, output]) {
 			const err = errArg.toString().trim()
@@ -1173,7 +1169,7 @@ export class PerforceContext {
 			}
 			else if (out.startsWith('No files to submit.')) {
 				this.logger.info(`=== SUBMIT FAIL === \nERR:${err}\nOUT:${out}`);
-				await this.deleteCl(workspace, changelist)
+				await this.deleteCl(roboWorkspace, changelist)
 				return 0
 			}
 			else if (out.startsWith('Submit validation failed')) {
@@ -1220,13 +1216,11 @@ export class PerforceContext {
 	// delete a CL
 	// output format is just error or not
 	deleteCl(roboWorkspace: RoboWorkspace, changelist: number, edgeServerAddress?: string) {
-		const workspace = coercePerforceWorkspace(roboWorkspace);
-		return this._execP4(workspace, ["change", "-d", changelist.toString()], {edgeServerAddress});
+		return this._execP4(roboWorkspace, ["change", "-d", changelist.toString()], {edgeServerAddress});
 	}
 
 	// run p4 'opened' command: lists files in changelist with details of edit state (e.g. if a copy, provides source path)
 	opened(roboWorkspace: RoboWorkspace | null, arg: number | string, exclusive?: boolean) {
-		const workspace = roboWorkspace && coercePerforceWorkspace(roboWorkspace);
 		const args = ['opened']
 		if (typeof arg === 'number') {
 			// see what files are open in the given changelist
@@ -1236,7 +1230,7 @@ export class PerforceContext {
 			// see which workspace has a file checked out/added
 			args.push(exclusive && perforceMultiServerEnvironment ? '-x' : '-a', arg)
 		}
-		return this.execAndParse(workspace, args) as Promise<OpenedFileRecord[]>
+		return this.execAndParse(roboWorkspace, args) as Promise<OpenedFileRecord[]>
 	}
 
 	async revertFile(file: string) {
@@ -1311,9 +1305,8 @@ export class PerforceContext {
 
 	// list files in changelist that need to be resolved
 	async listFilesToResolve(roboWorkspace: RoboWorkspace, changelist: number) {
-		const workspace = coercePerforceWorkspace(roboWorkspace);
 		try {
-			return await this.execAndParse(workspace, ['resolve', '-n', '-c', changelist.toString()]);
+			return await this.execAndParse(roboWorkspace, ['resolve', '-n', '-c', changelist.toString()]);
 		}
 		catch (err) {
 			if (err.toString().toLowerCase().includes('no file(s) to resolve')) {
@@ -1324,27 +1317,24 @@ export class PerforceContext {
 	}
 
 	async move(roboWorkspace: RoboWorkspace, cl: number, src: string, target: string, opts?: string[]) {
-		const workspace = coercePerforceWorkspace(roboWorkspace);
-		return this._execP4(workspace, ['move', ...(opts || []), '-c', cl.toString(), src, target]);
+		return this._execP4(roboWorkspace, ['move', ...(opts || []), '-c', cl.toString(), src, target]);
 	}
 
 	/**
 	 * Run a P4 command by name on one or more files (one call per file, run in parallel)
 	 */
 	run(roboWorkspace: RoboWorkspace, action: string, cl: number, filePaths: string[], opts?: string[]) {
-		const workspace = coercePerforceWorkspace(roboWorkspace);
 		const args = [action, ...(opts || []), '-c', cl.toString()]
 		return Promise.all(filePaths.map(
-			path => this._execP4(workspace, [...args, path])
+			path => this._execP4(roboWorkspace, [...args, path])
 		));
 	}
 
 	// shelve a CL
 	// output format is just error or not
 	async shelve(roboWorkspace: RoboWorkspace, changelist: number, edgeServerAddress?: string) {
-		const workspace = coercePerforceWorkspace(roboWorkspace)
 		try {
-			await this._execP4(workspace, ['shelve', '-f', '-c', changelist.toString()],
+			await this._execP4(roboWorkspace, ['shelve', '-f', '-c', changelist.toString()],
 				{ numRetries: 0, edgeServerAddress })
 		}
 		catch (reason) {
@@ -1365,9 +1355,8 @@ export class PerforceContext {
 	// shelve a CL
 	// output format is just error or not
 	async unshelve(roboWorkspace: RoboWorkspace, changelist: number) {
-		const workspace = coercePerforceWorkspace(roboWorkspace)
 		try {
-			await this._execP4(workspace, ['unshelve', '-s', changelist.toString(), '-f', '-c', changelist.toString()])
+			await this._execP4(roboWorkspace, ['unshelve', '-s', changelist.toString(), '-f', '-c', changelist.toString()])
 		}
 		catch (reason) {
 			if (!isExecP4Error(reason)) {
@@ -1385,8 +1374,7 @@ export class PerforceContext {
 
 	// delete shelved files from a CL
 	delete_shelved(roboWorkspace: RoboWorkspace, changelist: number) {
-		const workspace = coercePerforceWorkspace(roboWorkspace);
-		return this._execP4(workspace, ["shelve", "-d", "-c", changelist.toString(), "-f"]);
+		return this._execP4(roboWorkspace, ["shelve", "-d", "-c", changelist.toString(), "-f"]);
 	}
 
 	// get the email (according to P4) for a specific user
@@ -1402,26 +1390,24 @@ export class PerforceContext {
 
 	/** Check out a file into a specific changelist ** ASYNC ** */
 	edit(roboWorkspace: RoboWorkspace, cl: number, filePath: string, additionalArgs?: string[]) {
-		const workspace = coercePerforceWorkspace(roboWorkspace);
 		const args = [
 			'edit', 
 			'-c', cl.toString(), 
 			...(additionalArgs || []),
 			filePath
 		]
-		return this._execP4(workspace, args);
+		return this._execP4(roboWorkspace, args);
 	}
 
 	/** Add a file into a specific changelist ** ASYNC ** */
 	add(roboWorkspace: RoboWorkspace, cl: number, filePath: string, filetype?: string) {
-		const workspace = coercePerforceWorkspace(roboWorkspace);
 		const args = [
 			'add', 
 			'-c', cl.toString(), 
 			...(filetype ? ['-t', filetype] : []),
 			filePath
 		]
-		return this._execP4(workspace, args);
+		return this._execP4(roboWorkspace, args);
 	}
 
 	async describe(cl: number, maxFiles?: number, includeShelved?: boolean) {
@@ -1780,12 +1766,11 @@ export class PerforceContext {
 	}
 
 	static async _execP4Ztag(logger: ContextualLogger, roboWorkspace: RoboWorkspace, args: string[], opts?: ExecZtagOpts) {
-		const workspace = coercePerforceWorkspace(roboWorkspace);
 		if (opts && opts.format) {
-			return PerforceContext._execP4(logger, workspace, ['-ztag', '-F', opts.format, ...args], opts)
+			return PerforceContext._execP4(logger, roboWorkspace, ['-ztag', '-F', opts.format, ...args], opts)
 		}
 		else {
-			return parseZTag(await PerforceContext._execP4(logger, workspace, ['-ztag', ...args], opts), opts);
+			return parseZTag(await PerforceContext._execP4(logger, roboWorkspace, ['-ztag', ...args], opts), opts);
 		}
 	}
 
