don't do the FROM workspace name changing if the workspaceNameOverride is specified

diff --git a/Engine/Extras/RoboMerge/v3/src/robo/edgebot.ts b/Engine/Extras/RoboMerge/v3/src/robo/edgebot.ts
--- a/Engine/Extras/RoboMerge/v3/src/robo/edgebot.ts
+++ b/Engine/Extras/RoboMerge/v3/src/robo/edgebot.ts
@@ -460,13 +460,15 @@ class EdgeBotImpl extends PerforceStatefulBot {
 
 		// name the workspace
 		let workspaceName = this.options.workspaceNameOverride || ['ROBOMERGE', this.targetBranch.parent.botname, this.targetBranch.name].join('_');
-		let fromName = `_FROM_${(this.targetBranch.parent == this.sourceBranch.parent ? '' : this.sourceBranch.parent.botname)}_${this.sourceBranch.name}`
-		if (this.incognitoMode) {
-			const hashCode = (s: string) => s.split('').reduce((a,b)=>{a=((a<<5)-a)+b.charCodeAt(0);return a&a},0)
-			workspaceName += '_' + hashCode(fromName)
-		}
-		else {
-			workspaceName += fromName
+		if (!this.options.workspaceNameOverride) {
+			let fromName = `_FROM_${(this.targetBranch.parent == this.sourceBranch.parent ? '' : this.sourceBranch.parent.botname)}_${this.sourceBranch.name}`
+			if (this.incognitoMode) {
+				const hashCode = (s: string) => s.split('').reduce((a,b)=>{a=((a<<5)-a)+b.charCodeAt(0);return a&a},0)
+				workspaceName += '_' + hashCode(fromName)
+			}
+			else {
+				workspaceName += fromName
+			}
 		}
 		const p4username = getPerforceUsername()
 		if (p4username !== 'robomerge') {
