Properly fix setting owner on edge server

diff --git a/Engine/Extras/RoboMerge/v3/src/robo/edgebot.ts b/Engine/Extras/RoboMerge/v3/src/robo/edgebot.ts
--- a/Engine/Extras/RoboMerge/v3/src/robo/edgebot.ts
+++ b/Engine/Extras/RoboMerge/v3/src/robo/edgebot.ts
@@ -741,13 +741,14 @@ class EdgeBotImpl extends PerforceStatefulBot {
 		const info = pending.change
 		const target = pending.action
 		const changenum = pending.newCl
+		const edgeServerAddress = pending.change.edgeServerToHostShelf && pending.change.edgeServerToHostShelf.address
 
 		// try to submit
 		this._log_action(`Submitting CL ${changenum} by ${info.author}`)
 		const result = await this.p4.submit(
 			pending.change.targetWorkspace, 
 			changenum,
-			pending.change.edgeServerToHostShelf && pending.change.edgeServerToHostShelf.address)
+			edgeServerAddress)
 			
 		if (typeof result === "string") {
 			// an error occurred while submitting
@@ -780,7 +781,7 @@ class EdgeBotImpl extends PerforceStatefulBot {
 			// change owner, so users can edit change descriptions later for reconsideration
 			this.edgeBotLogger.info(`Setting owner of CL ${finalCl} to author of change: ${info.author}`);
 			try {
-				await this.p4.editOwner(null, finalCl, info.author, {changeSubmitted: true})
+				await this.p4.editOwner(info.targetWorkspace, finalCl, info.author, {changeSubmitted: true, edgeServerAddress})
 			}
 			catch (reason) {
 				let errPreface = 'Error changing owner'
