Allow triagers to specified using preformatted tags to deal with the fact the auto-formatting is not working when the channel is added to multiple workspaces

diff --git a/Engine/Extras/RoboMerge/v3/src/robo/nodebot.ts b/Engine/Extras/RoboMerge/v3/src/robo/nodebot.ts
--- a/Engine/Extras/RoboMerge/v3/src/robo/nodebot.ts
+++ b/Engine/Extras/RoboMerge/v3/src/robo/nodebot.ts
@@ -2307,7 +2307,7 @@ export class NodeBot extends PerforceStatefulBot implements NodeBotInterface {
 				
 						let triager = getNagProperty("triager")
 
-						if (triager && !triager.startsWith('@'))
+						if (triager && !triager.startsWith('@') && !triager.startsWith('<'))
 						{
 							const emailAddress = await this.findEmail(triager)
 							if (emailAddress)
