Simplify clean workspace logic and eliminate unnecessary log spew

diff --git a/Engine/Extras/RoboMerge/v3/src/common/p4util.ts b/Engine/Extras/RoboMerge/v3/src/common/p4util.ts
--- a/Engine/Extras/RoboMerge/v3/src/common/p4util.ts
+++ b/Engine/Extras/RoboMerge/v3/src/common/p4util.ts
@@ -174,27 +174,11 @@ export async function convertIntegrateToEdit(p4: PerforceContext, roboWorkspace:
 	}
 }
 
-export async function cleanWorkspaces(p4: PerforceContext, workspaces: [string, string][], edgeServerAddress?: string) {
-	const changes = await p4.get_pending_changes() as Change[]
-
-	const nameSet = new Set(workspaces.map(([name, _]) => name))
-	const toRevert: Change[] = []
-	const ignoredChangeCounts = new Map<string, number>()
-	for (const change of changes) {
-		if (nameSet.has(change.client)) {
-			toRevert.push(change);
-		}
-		else {
-			ignoredChangeCounts.set(change.client, (ignoredChangeCounts.get(change.client) || 0) + 1);
-		}
-	}
-
-	for (const [ws, count] of ignoredChangeCounts.entries()) {
-		p4utilsLogger.info(`Ignoring ${count} changelist${count > 1 ? 's' : ''} in ${ws} which is not monitored by this bot`);
-	}
+export async function cleanWorkspace(p4: PerforceContext, workspaceName: string, workspaceRoot: string, edgeServerAddress?: string) {
+	const changes = await p4.get_pending_changes(workspaceName) as Change[]
 
 	// revert any pending changes left over from previous runs
-	for (const change of toRevert) {
+	for (const change of changes) {
 		const changeStr = `CL ${change.change}: ${change.desc}`
 		const workspace = change.client
 		if (change.shelved) {
@@ -220,7 +204,7 @@ export async function cleanWorkspaces(p4: PerforceContext, workspaces: [string,
 	}
 
 	p4utilsLogger.info('Resetting all workspaces to revision 0')
-	await Promise.all(workspaces.map(([name, root]) => p4.sync(name, root + '#0', {edgeServerAddress})))
+	await p4.sync(workspaceName, workspaceRoot + '#0', {edgeServerAddress})
 }
 
 export async function getWorkspacesForUser(p4: PerforceContext, user: string) {
diff --git a/Engine/Extras/RoboMerge/v3/src/robo/autobranchupdater.ts b/Engine/Extras/RoboMerge/v3/src/robo/autobranchupdater.ts
--- a/Engine/Extras/RoboMerge/v3/src/robo/autobranchupdater.ts
+++ b/Engine/Extras/RoboMerge/v3/src/robo/autobranchupdater.ts
@@ -193,7 +193,7 @@ export class AutoBranchUpdater implements Bot {
 
 					if (retryCount < 5) {
 						this.abuLogger.warn(`Mirror file reloading error (retries ${retryCount}: ${err}`)
-						p4util.cleanWorkspaces(this.p4, [[mirrorWorkspace.name, mirrorWorkspace.depotpath]])
+						p4util.cleanWorkspace(this.p4, mirrorWorkspace.name, mirrorWorkspace.depotpath)
 					}
 					else {
 						this.abuLogger.error('Mirror file reloading failed, retries exhausted')
diff --git a/Engine/Extras/RoboMerge/v3/src/robo/edgebot.ts b/Engine/Extras/RoboMerge/v3/src/robo/edgebot.ts
--- a/Engine/Extras/RoboMerge/v3/src/robo/edgebot.ts
+++ b/Engine/Extras/RoboMerge/v3/src/robo/edgebot.ts
@@ -492,7 +492,7 @@ class EdgeBotImpl extends PerforceStatefulBot {
 			if (existingWorkspaceInfo[0].IsUnloaded) {
 				await this.p4.reloadWorkspace(workspaceName, edgeServerAddress)
 			}
-			await p4util.cleanWorkspaces(this.p4, [[workspaceName, this.targetBranch.rootPath]], edgeServerAddress)
+			await p4util.cleanWorkspace(this.p4, workspaceName, this.targetBranch.rootPath, edgeServerAddress)
 		}
 		else {
 			const params: any = {};
