Remove not useful graph interface intermediary layer

diff --git a/Engine/Extras/RoboMerge/v3/src/robo/autobranchupdater.ts b/Engine/Extras/RoboMerge/v3/src/robo/autobranchupdater.ts
--- a/Engine/Extras/RoboMerge/v3/src/robo/autobranchupdater.ts
+++ b/Engine/Extras/RoboMerge/v3/src/robo/autobranchupdater.ts
@@ -6,7 +6,7 @@ import { ContextualLogger } from '../common/logger';
 import { P4_FORCE, PerforceContext, Workspace } from '../common/perforce';
 import { Bot } from './bot-interfaces';
 import { BranchDefs } from './branchdefs';
-import { GraphInterface } from './graph-interface';
+import { GraphBot } from './graphbot';
 import * as p4util from '../common/p4util';
 
 const DISABLE = false
@@ -48,7 +48,7 @@ export class AutoBranchUpdater implements Bot {
 
 	tickCount = 0
 
-	constructor(private graphBot: GraphInterface, parentLogger: ContextualLogger, private readonly previewMode: boolean) {
+	constructor(private graphBot: GraphBot, parentLogger: ContextualLogger, private readonly previewMode: boolean) {
 		this.p4 = AutoBranchUpdater.p4!
 
 		const config = AutoBranchUpdater.config
@@ -248,7 +248,7 @@ export class AutoBranchUpdater implements Bot {
 		await this.p4.submit(workspace, cl)
 	}
 
-	static getMirrorWorkspace(bot: GraphInterface): MirrorPaths | null {
+	static getMirrorWorkspace(bot: GraphBot): MirrorPaths | null {
 		const mirrorPathBits = bot.branchGraph.config.mirrorPath
 		if (mirrorPathBits.length === 0) {
 			return null
diff --git a/Engine/Extras/RoboMerge/v3/src/robo/graph-interface.ts b/Engine/Extras/RoboMerge/v3/src/robo/graph-interface.ts
deleted file mode 100644
diff --git a/Engine/Extras/RoboMerge/v3/src/robo/graphbot.ts b/Engine/Extras/RoboMerge/v3/src/robo/graphbot.ts
--- a/Engine/Extras/RoboMerge/v3/src/robo/graphbot.ts
+++ b/Engine/Extras/RoboMerge/v3/src/robo/graphbot.ts
@@ -12,7 +12,6 @@ import { BotConfig, BranchDefs, BranchGraphDefinition } from './branchdefs';
 import { BranchGraph } from './branchgraph';
 import { PersistentConflict } from './conflict-interfaces';
 import { BotEventHandler, BotEventTriggers } from './events';
-import { GraphInterface } from './graph-interface';
 import { NodeBot } from './nodebot';
 import { bindBotNotifications, BotNotifications, NOTIFICATIONS_PERSISTENCE_KEY } from './notifications';
 import { postMessageToChannel, postToRobomergeAlerts, SlackMessages } from './notifications';
@@ -28,7 +27,7 @@ import { GraphAPI } from '../new/graph';
 const MAX_ERROR_LENGTH_TO_REPORT = 2000
 
 export type ReloadListeners = (graphBot: GraphBot, logger: ContextualLogger) => void
-export class GraphBot implements GraphInterface, BotEventHandler {
+export class GraphBot implements BotEventHandler {
 	static dataDirectory: string
 	branchGraph: BranchGraph
 	filename: string
