Encapsulate the vault in to an object

diff --git a/Engine/Extras/RoboMerge/v3/src/robo/notifications.ts b/Engine/Extras/RoboMerge/v3/src/robo/notifications.ts
--- a/Engine/Extras/RoboMerge/v3/src/robo/notifications.ts
+++ b/Engine/Extras/RoboMerge/v3/src/robo/notifications.ts
@@ -1,6 +1,5 @@
 // Copyright Epic Games, Inc. All Rights Reserved.
 
-import * as fs from 'fs';
 import { DateTimeFormatOptions } from 'intl';
 import { Args } from '../common/args';
 import { Badge } from '../common/badge';
@@ -13,11 +12,10 @@ import { NodeBot } from './nodebot';
 import { BlockageNodeOpUrls } from './roboserver';
 import { Context } from './settings';
 import { Slack, SlackAttachment, SlackFile, SlackLinkButtonAction, SlackLinkButtonsAttachment, SlackMessageField, SlackMessage, SlackMessageStyles } from './slack';
+import { Vault } from './vault'
 import { WebServer } from '../common/webserver'
 import { DummySlackApp } from '../common/dummyslackserver'
 
-const notificationsStartupLogger = new ContextualLogger('Notifications Startup')
-
 export const NOTIFICATIONS_PERSISTENCE_KEY = 'notifications'
 export const SLACK_MESSAGES_PERSISTENCE_KEY = 'slackMessages'
 const CHANGELIST_FIELD_TITLE = 'Change'
@@ -32,7 +30,7 @@ let SLACK_TOKENS: {[name: string]: string} = {}
 const SLACK_DEV_DUMMY_TOKEN = 'dev'
 
 let args: Args
-export function notificationsInit(inArgs: Args) {
+export function notificationsInit(inArgs: Args, vault: Vault) {
 	
 	args = inArgs
 
@@ -47,18 +45,7 @@ export function notificationsInit(inArgs: Args) {
 		return
 	}
 
-	let vault
-	try {
-		const vaultString = fs.readFileSync(args.vault + '/vault.json', 'ascii')
-		vault = JSON.parse(vaultString)
-	}
-	catch (err) {
-		const errStr = err.toString().replace(/(E|e)(R|r)(R|r)(O|o)(R|r)/,"$1$2$30$5")
-		notificationsStartupLogger.warn(`Warning, failed to find Slack secrets in vault (ok in dev): ${errStr}`)
-		return
-	}
-
-	const tokensObj = vault['slack-tokens']
+	const tokensObj = vault.slackTokens
 	if (tokensObj) {
 		SLACK_TOKENS = tokensObj
 	}
diff --git a/Engine/Extras/RoboMerge/v3/src/robo/robo.ts b/Engine/Extras/RoboMerge/v3/src/robo/robo.ts
--- a/Engine/Extras/RoboMerge/v3/src/robo/robo.ts
+++ b/Engine/Extras/RoboMerge/v3/src/robo/robo.ts
@@ -24,6 +24,7 @@ import { roboAnalytics, setGlobalAnalytics } from './roboanalytics';
 import { RoboServer } from './roboserver';
 import { notificationsInit } from './notifications'
 import { settingsInit } from './settings'
+import { Vault } from './vault'
 import * as Preview from './preview'
 
 /*************************
@@ -210,7 +211,8 @@ if (!maybeNullArgs) {
 }
 const args = maybeNullArgs!
 
-notificationsInit(args)
+const vault = new Vault(args.vault)
+notificationsInit(args, vault)
 settingsInit(args)
 Preview.init(args.branchSpecsRootPath)
 
diff --git a/Engine/Extras/RoboMerge/v3/src/robo/session.ts b/Engine/Extras/RoboMerge/v3/src/robo/session.ts
--- a/Engine/Extras/RoboMerge/v3/src/robo/session.ts
+++ b/Engine/Extras/RoboMerge/v3/src/robo/session.ts
@@ -4,6 +4,7 @@ import * as crypto from 'crypto';
 import * as fs from 'fs';
 import { Arg, readProcessArgs } from '../common/args';
 import { ContextualLogger } from '../common/logger';
+import { Vault } from './vault'
 
 const LdapAuth: any = require('ldapauth-fork');
 
@@ -70,19 +71,14 @@ export class Session {
 			}
 		}
 
-		try {
-			const vaultString = fs.readFileSync(Session.VAULT_PATH + '/vault.json', 'ascii')
-			const vault = JSON.parse(vaultString)
-			Session.LDAP_CONFIG.adminPassword = vault['ldap-password']
-			Session.COOKIE_KEY = vault['cookie-key']
+		const vault = new Vault(Session.VAULT_PATH, logger)
+		if (vault.valid) {
+			Session.LDAP_CONFIG.adminPassword = vault.ldapPassword
+			Session.COOKIE_KEY = vault.cookieKey
 			return
 		}
-		catch (err) {
-			const errStr = err.toString().replace(/(E|e)(R|r)(R|r)(O|o)(R|r)/,"$1$2$30$5")
-			logger.warn(`Warning (ok in dev):  ${errStr}`)
-			Session.COOKIE_KEY = DEV_COOKIE_KEY
-		}
 
+		Session.COOKIE_KEY = DEV_COOKIE_KEY
 		logger.warn('No vault or no LDAP creds in vault (this is ok for testing)')
 	}
 
diff --git a/Engine/Extras/RoboMerge/v3/src/robo/vault.ts b/Engine/Extras/RoboMerge/v3/src/robo/vault.ts
new file mode 100644
--- /dev/null
+++ b/Engine/Extras/RoboMerge/v3/src/robo/vault.ts
@@ -0,0 +1,48 @@
+// Copyright Epic Games, Inc. All Rights Reserved.
+
+import * as fs from 'fs';
+import { ContextualLogger } from '../common/logger';
+
+export class Vault {
+
+    private _vault: any
+    readonly valid: boolean = false
+
+    constructor(vaultPath: string, vaultLogger?: ContextualLogger) {
+
+        if (!vaultLogger) {
+            vaultLogger = new ContextualLogger('Vault Init')
+        }
+
+        try {
+            var vaultString = fs.readFileSync(vaultPath + '/vault.json', 'ascii')
+        }
+        catch (err) {
+            const errStr = err.toString().replace(/(E|e)(R|r)(R|r)(O|o)(R|r)/,"$1$2$30$5")
+            vaultLogger.warn(`Warning, failed to read vault (ok in dev): ${errStr}`)
+            return
+        }
+        
+        try {
+            this._vault = JSON.parse(vaultString)
+        }
+        catch (err) {
+            vaultLogger.error(`Failed to parse vault file: ${err}`)
+            return
+        }
+
+        this.valid = true
+    }
+
+    get slackTokens() {
+        return this._vault['slack-tokens']
+    }
+
+    get cookieKey() {
+        return this._vault['cookie-key']
+    }
+
+    get ldapPassword() {
+        return this._vault['ldap-password']
+    }
+}
\ No newline at end of file
