Eliminate globals and encapsulate perforce server details
Remove noUsername from ExecOpts and allow all calls to manage adding -u themselves as appropriate

diff --git a/Engine/Extras/RoboMerge/v3/src/common/perforce.ts b/Engine/Extras/RoboMerge/v3/src/common/perforce.ts
--- a/Engine/Extras/RoboMerge/v3/src/common/perforce.ts
+++ b/Engine/Extras/RoboMerge/v3/src/common/perforce.ts
@@ -268,7 +268,6 @@ interface ExecOpts {
 	stdin?: string
 	quiet?: boolean
 	noCwd?: boolean
-	noUsername?: boolean
 	numRetries?: number
 	trace?: boolean
 	edgeServerAddress?: string
@@ -316,6 +315,12 @@ export interface ConflictedResolveNFile {
 	contentResolveType?: string  // e.g. '3waytext', '2wayraw' -- not applicable for delete/branch merges
 }
 
+interface ServerConfig {
+	username: string
+	serverVersion: number
+	multiServerEnvironment: boolean
+}
+
 export class ResolveResult {
 	private rResultLogger: ContextualLogger
 	private resolveOutput: any[] | null
@@ -442,23 +447,11 @@ export const getRunningPerforceCommands = () => [...runningPerforceCommands]
 export const P4_FORCE = '-f';
 const robomergeVersion = `v${VersionReader.getBuildNumber()}`
 
-// check if we are logged in and p4 is set up correctly
-let perforceUsername: string
-export function getPerforceUsername() {
-	if (!perforceUsername) {
-		throw new Error('Attempting to access Perforce username information before login completed.')
-	}
-	return perforceUsername
-}
-
-let perforceMultiServerEnvironment: boolean
-let perforceServerVersion: number
-
 const MIN_SERVER_VERSION: number = 2020.1
 
 /**
- * This method must succeed before PerforceContext can be used. Otherwise retrieving the Perforce username through
- * getPerforceUsername() will error.
+ * This method must succeed before PerforceContext can be used. Otherwise retrieving the Perforce context through
+ * getServerContext() will error.
  */
 export async function initializePerforce(logger: ContextualLogger) {
 	setInterval(() => {
@@ -471,25 +464,8 @@ export async function initializePerforce(logger: ContextualLogger) {
 		}
 	}, 5*60*1000)
 
-	const output = await PerforceContext.execAndParse(logger, null, ["login", "-s"], { noUsername: true });
-	let resp = output[0];
-
-	if (resp && resp.User) {
-		perforceUsername = resp.User;
-
-		const serverVersion = await PerforceContext._execP4(logger, null, ["-ztag","-F","%serverVersion%","info"])
-		const match = serverVersion.match(/.*\/(\d+\.\d+)\/.*/)
-		if (!match) {
-			throw new Error(`Unable to parse server version from ${serverVersion}`)
-		}
-		perforceServerVersion = parseFloat(match[1])
-		if (perforceServerVersion < MIN_SERVER_VERSION) {
-			throw new Error(`Robomerge requires a minimum server version of ${MIN_SERVER_VERSION}`)
-		}
-
-		const serversOutput = await PerforceContext.execAndParse(logger, null, ["servers"]);
-		perforceMultiServerEnvironment = serversOutput.length > 1
-	}
+	logger.info("Robomerge running in single server mode")
+	await PerforceContext.initServer(logger, "default")
 }
 
 /**
@@ -497,10 +473,71 @@ export async function initializePerforce(logger: ContextualLogger) {
  * internal serialization object.
  */
 export class PerforceContext {
-	username = getPerforceUsername()
 
 	// Use same logger as instance owner, to cut down on context length
-	constructor(private readonly logger: ContextualLogger) {
+	private constructor(private readonly logger: ContextualLogger, private _serverConfig: ServerConfig) {
+	}
+
+	static readonly servers: Map<string, ServerConfig> = new Map()
+
+	static async initServer(logger: ContextualLogger, serverID: string, serverDetails?: any) {
+		
+		const serverAddress = serverDetails && serverDetails.port
+
+		const output = await PerforceContext.execAndParse(logger, null, ["login", "-s"], { edgeServerAddress: serverAddress });
+		let resp = output[0];
+	
+		if (resp && resp.User) {
+			const rawServerVersion = await PerforceContext._execP4(logger, null, ["-u", resp.User, "-ztag","-F","%serverVersion%","info"], { edgeServerAddress: serverAddress} )
+			const match = rawServerVersion.match(/.*\/(\d+\.\d+)\/.*/)
+			if (!match) {
+				throw new Error(`Unable to parse server version from ${rawServerVersion}`)
+			}
+			let serverVersion = parseFloat(match[1])
+			if (serverVersion < MIN_SERVER_VERSION) {
+				throw new Error(`Robomerge requires a minimum server version of ${MIN_SERVER_VERSION}`)
+			}
+	
+			const serversOutput = await PerforceContext.execAndParse(logger, null, ["-u", resp.User, "servers"], { edgeServerAddress: serverAddress });
+
+			let serverConfig: ServerConfig = {
+				username: resp.User,
+				serverVersion,
+				multiServerEnvironment: serversOutput.length > 1
+			}
+			this.servers.set(serverID, serverConfig)
+
+			if (serverAddress) {
+				logger.info(`P4 server ${serverID}: ${serverAddress} initialized. User=${resp.User}`)
+			}
+			else {
+				logger.info(`P4 initialized. User=${resp.User}`)
+			}
+		}
+		else {
+			throw new Error(`Failed to login to ${serverAddress || "server"}`)
+		}
+	}
+
+	static getServerContext(logger: ContextualLogger, serverID?: string) {
+		serverID = serverID || 'default'
+		const server = this.servers.get(serverID)
+		if (!server) {
+			throw new Error(`Unknown server ${serverID}`)
+		}
+		return new PerforceContext(logger, server)
+	}
+
+	get username() {
+		return this._serverConfig.username
+	}
+
+	get serverVersion() {
+		return this._serverConfig.serverVersion
+	}
+
+	get multiServerEnvironment() {
+		return this._serverConfig.multiServerEnvironment
 	}
 
 	// get a list of all pending changes for this user
@@ -938,7 +975,7 @@ export class PerforceContext {
 		let noSuchFilesPossible = false // Helper variable for error catching
 		// Branchspec -- takes priority above other possibilities
 		if (source.branchspec) {
-			if (perforceServerVersion >= 2024.1) {
+			if (this.serverVersion >= 2024.1) {
 				// 2024.1 now prevents merges between streams by default
 				cmdList.push("-F")
 			}
@@ -961,7 +998,7 @@ export class PerforceContext {
 		}
 		// Basic branch to branch
 		else {
-			if (perforceServerVersion >= 2024.1) {
+			if (this.serverVersion >= 2024.1) {
 				// 2024.1 now prevents merges between streams by default
 				cmdList.push("-F")
 			}
@@ -1228,7 +1265,7 @@ export class PerforceContext {
 		}
 		else {
 			// see which workspace has a file checked out/added
-			args.push(exclusive && perforceMultiServerEnvironment ? '-x' : '-a', arg)
+			args.push(exclusive && this.multiServerEnvironment ? '-x' : '-a', arg)
 		}
 		return this.execAndParse(roboWorkspace, args) as Promise<OpenedFileRecord[]>
 	}
@@ -1593,11 +1630,6 @@ export class PerforceContext {
 
 		const opts = optsIn || {}
 
-		// Perforce can get mighty confused if you log into multiple accounts at any point (only relevant to local debugging)
-		if (!opts.noUsername) {
-			args = ['-u', getPerforceUsername(), ...args]
-		}
-
 		if (workspace && workspace.name) {
 			args = ['-c', workspace.name, ...args]
 		}
@@ -1762,7 +1794,7 @@ export class PerforceContext {
 	}
 
 	private _execP4(roboWorkspace: RoboWorkspace, args: string[], optsIn?: ExecOpts) {
-		return PerforceContext._execP4(this.logger, roboWorkspace, args, optsIn)
+		return PerforceContext._execP4(this.logger, roboWorkspace, ['-u', this.username, ...args], optsIn)
 	}
 
 	static async _execP4Ztag(logger: ContextualLogger, roboWorkspace: RoboWorkspace, args: string[], opts?: ExecZtagOpts) {
@@ -1775,7 +1807,7 @@ export class PerforceContext {
 	}
 
 	private async _execP4Ztag(roboWorkspace: RoboWorkspace, args: string[], opts?: ExecZtagOpts) {
-		return PerforceContext._execP4Ztag(this.logger, roboWorkspace, args, opts)
+		return PerforceContext._execP4Ztag(this.logger, roboWorkspace, ['-u', this.username, ...args], opts)
 	}
 
 	static parseValue(key: string, value: any, parseOptions?: ParseOptions)
@@ -1848,7 +1880,7 @@ export class PerforceContext {
 	}
 
 	async execAndParse(roboWorkspace: RoboWorkspace, args: string[], execOptions?: ExecOpts, parseOptions?: ParseOptions) {
-		return PerforceContext.execAndParse(this.logger, roboWorkspace, args, execOptions, parseOptions)
+		return PerforceContext.execAndParse(this.logger, roboWorkspace, ['-u', this.username, ...args], execOptions, parseOptions)
 	}
 
 	static async execAndParseArray(logger: ContextualLogger, roboWorkspace: RoboWorkspace, args: string[], execOptions?: ExecOpts, headerOptions?: ParseOptions, arrayEntryOptions?: ParseOptions) {
diff --git a/Engine/Extras/RoboMerge/v3/src/new/graph.ts b/Engine/Extras/RoboMerge/v3/src/new/graph.ts
--- a/Engine/Extras/RoboMerge/v3/src/new/graph.ts
+++ b/Engine/Extras/RoboMerge/v3/src/new/graph.ts
@@ -813,7 +813,7 @@ export class Trace {
 	constructor(private graph: Graph, parentLogger: ContextualLogger)
 	{
 		this.traceLogger = parentLogger.createChild('Trace')
-		this.p4 = new PerforceContext(this.traceLogger)
+		this.p4 = PerforceContext.getServerContext(this.traceLogger)
 	}
 
 	async traceRoute(cl: number, sourceNodeStr: string, targetNodeStr: string) {
diff --git a/Engine/Extras/RoboMerge/v3/src/robo/autobranchupdater.ts b/Engine/Extras/RoboMerge/v3/src/robo/autobranchupdater.ts
--- a/Engine/Extras/RoboMerge/v3/src/robo/autobranchupdater.ts
+++ b/Engine/Extras/RoboMerge/v3/src/robo/autobranchupdater.ts
@@ -42,14 +42,13 @@ export class AutoBranchUpdater implements Bot {
 	isRunning = false
 	isActive = false
 
-	static p4: PerforceContext
 	static config: AutoBranchUpdaterConfig
 	static initialCl: number
 
 	tickCount = 0
 
 	constructor(private graphBot: GraphBot, parentLogger: ContextualLogger, private readonly previewMode: boolean) {
-		this.p4 = AutoBranchUpdater.p4!
+		this.p4 = PerforceContext.getServerContext(parentLogger)
 
 		const config = AutoBranchUpdater.config
 		this.filePath = `${config.rootPath}/${graphBot.filename}`
@@ -59,14 +58,13 @@ export class AutoBranchUpdater implements Bot {
 		this.abuLogger = parentLogger.createChild('ABU')
 	}
 
-	static async init(deps: any, config: AutoBranchUpdaterConfig, parentLogger: ContextualLogger) {
-		this.p4 = deps.p4
+	static async init(p4: PerforceContext, config: AutoBranchUpdaterConfig, parentLogger: ContextualLogger) {
 		this.config = config
 		const logger = parentLogger.createChild('ABU')
 
 		logger.info(`Finding most recent branch spec changelist from ${config.rootPath}`)
 		var bsRoot = config.rootPath + '/...'
-		const change = await this.p4.latestChange(bsRoot)
+		const change = await p4.latestChange(bsRoot)
 		if (change === null)
 			throw new Error(`Unable to query for most recent branch specs CL`)
 
@@ -85,7 +83,7 @@ export class AutoBranchUpdater implements Bot {
 		}
 
 		const opts = !config.devMode ? [P4_FORCE] : undefined 
-		await this.p4.sync(config.workspace, `${bsRoot}@${change.change}`, {opts})
+		await p4.sync(config.workspace, `${bsRoot}@${change.change}`, {opts})
 	}
 	
 	async start() {
diff --git a/Engine/Extras/RoboMerge/v3/src/robo/badges.ts b/Engine/Extras/RoboMerge/v3/src/robo/badges.ts
--- a/Engine/Extras/RoboMerge/v3/src/robo/badges.ts
+++ b/Engine/Extras/RoboMerge/v3/src/robo/badges.ts
@@ -52,7 +52,7 @@ class BadgeHandler implements BotEventHandler {
 		this.externalUrl = externalUrl
 
 		this.badgeHandlerLogger = parentLogger.createChild('BadgeHandler')
-		this.p4 = new PerforceContext(this.badgeHandlerLogger)
+		this.p4 = PerforceContext.getServerContext(this.badgeHandlerLogger)
 	}
 
 	/** See if this change implies badges should be added or updated */
diff --git a/Engine/Extras/RoboMerge/v3/src/robo/edgebot.ts b/Engine/Extras/RoboMerge/v3/src/robo/edgebot.ts
--- a/Engine/Extras/RoboMerge/v3/src/robo/edgebot.ts
+++ b/Engine/Extras/RoboMerge/v3/src/robo/edgebot.ts
@@ -5,7 +5,7 @@ import * as p4util from '../common/p4util';
 
 import { ContextualLogger } from "../common/logger";
 import { Recipients } from "../common/mailer";
-import { Change, coercePerforceWorkspace, ConflictedResolveNFile, EdgeServer, EditChangeOpts, EXCLUSIVE_CHECKOUT_REGEX, getPerforceUsername, getRootDirectoryForBranch, IntegrationSource, IntegrationTarget, isExecP4Error, OpenedFileRecord, PerforceContext } from "../common/perforce";
+import { Change, coercePerforceWorkspace, ConflictedResolveNFile, EdgeServer, EditChangeOpts, EXCLUSIVE_CHECKOUT_REGEX, getRootDirectoryForBranch, IntegrationSource, IntegrationTarget, isExecP4Error, OpenedFileRecord, PerforceContext } from "../common/perforce";
 import { VersionReader } from "../common/version";
 import { EdgeBotInterface, IPCControls, ReconsiderArgs } from "./bot-interfaces";
 import { AlreadyIntegrated, Branch, ChangeInfo, ConflictingFile, ExclusiveFile, Failure, MergeAction, PendingChange } from "./branch-interfaces";
@@ -115,7 +115,7 @@ class EdgeBotImpl extends PerforceStatefulBot {
 		this.graphBotName = sourceNode.graphBotName
 		this.edgeBotLogger = new ContextualLogger(this.fullNameForLogging)
 
-		this.p4 = new PerforceContext(this.edgeBotLogger)
+		this.p4 = PerforceContext.getServerContext(this.edgeBotLogger)
 		this.eventTriggers.onChangeParsed((info: ChangeInfo) => this.onGlobalChange(info))
 		this.gate = new Gate(
 			{ from: this.sourceBranch
@@ -470,7 +470,7 @@ class EdgeBotImpl extends PerforceStatefulBot {
 				workspaceName += fromName
 			}
 		}
-		const p4username = getPerforceUsername()
+		const p4username = this.p4.username
 		if (p4username !== 'robomerge') {
 			workspaceName = [p4username!.toUpperCase(), process.platform.toUpperCase(), workspaceName].join('_')
 		}
diff --git a/Engine/Extras/RoboMerge/v3/src/robo/graphbot.ts b/Engine/Extras/RoboMerge/v3/src/robo/graphbot.ts
--- a/Engine/Extras/RoboMerge/v3/src/robo/graphbot.ts
+++ b/Engine/Extras/RoboMerge/v3/src/robo/graphbot.ts
@@ -53,7 +53,7 @@ export class GraphBot implements BotEventHandler {
 		}
 
 		graphBot.botLogger = new ContextualLogger(botname.toUpperCase())
-		graphBot.p4 = new PerforceContext(graphBot.botLogger)
+		graphBot.p4 = PerforceContext.getServerContext(graphBot.botLogger)
 
 		graphBot.branchGraph = new BranchGraph(botname)
 		graphBot.filename = botname + '.branchmap.json'
diff --git a/Engine/Extras/RoboMerge/v3/src/robo/ipc.ts b/Engine/Extras/RoboMerge/v3/src/robo/ipc.ts
--- a/Engine/Extras/RoboMerge/v3/src/robo/ipc.ts
+++ b/Engine/Extras/RoboMerge/v3/src/robo/ipc.ts
@@ -1,7 +1,7 @@
 // Copyright Epic Games, Inc. All Rights Reserved.
 
 import { ContextualLogger, isNpmLogLevel, NpmLogLevelCompare, NpmLogLevelValues } from '../common/logger';
-import { DescribeResult, getRunningPerforceCommands } from '../common/perforce';
+import { DescribeResult, getRunningPerforceCommands, PerforceContext } from '../common/perforce';
 import { Trace } from '../new/graph';
 import { IPCControls, EdgeBotInterface, NodeBotInterface } from './bot-interfaces';
 import { OperationResult } from './branch-interfaces';
@@ -149,7 +149,7 @@ export class IPC {
 	}
 
 	private async getWorkspaces(username: string): Promise<OperationReturnType> {
-		const workspaces = await p4util.getWorkspacesForUser(this.robo.p4, username)
+		const workspaces = await p4util.getWorkspacesForUser(PerforceContext.getServerContext(this.robo.roboMergeLogger), username)
 		
 		if (!workspaces) {
 			return {
@@ -233,12 +233,13 @@ export class IPC {
 			}
 		}
 
+		let p4 = PerforceContext.getServerContext(this.robo.roboMergeLogger)
 		let changes = new Map<number, any>()
 
 		let data: any = {originalCL: cl, changes: {}}
 
 		let gatherCLInfo = async (cl: number, opts?: any) => {
-			let desc = await this.robo.p4.describe(cl, 1)
+			let desc = await p4.describe(cl, 1)
 			const mergeMethod = cl != data.originalCL ? getMergeMethod(desc.description) : 'initialSubmit'
 			const clNode = getNode(desc)
 			if (clNode || (RobomergeMethodStrings as readonly string[]).includes(mergeMethod)) {
@@ -331,7 +332,7 @@ export class IPC {
 				const entry = changeToConsider.desc.entries[i]
 				// integrated for move/delete will point at the paired move/add not where it was merged to in another stream, so not useful to evaluate it
 				if (entry.action != 'move/delete') {
-					const integrated = await this.robo.p4.integrated(null, entry.depotFile, {intoOnly: true, startCL: clToConsider})
+					const integrated = await p4.integrated(null, entry.depotFile, {intoOnly: true, startCL: clToConsider})
 					if (integrated.length > 0) {
 						for (let integ of integrated) {
 							const startToRev = integ.startToRev == "#none" ? 0 : parseInt(integ.startToRev.slice(1))
@@ -357,13 +358,13 @@ export class IPC {
 					// get more of the files. If the path is a root directory try and get a sampling 
 					// of files, otherwise just get an additional block to evaluate
 					if (changeToConsider.desc.path.endsWith("/...")) {
-						const dirs = await this.robo.p4.dirs(`${changeToConsider.desc.path.slice(0,-3)}*@=${clToConsider}`)
+						const dirs = await p4.dirs(`${changeToConsider.desc.path.slice(0,-3)}*@=${clToConsider}`)
 						for (const dir of dirs) {
-							changeToConsider.desc.entries.push(...(await this.robo.p4.files(`${dir}/...@=${clToConsider}`, 1)))
+							changeToConsider.desc.entries.push(...(await p4.files(`${dir}/...@=${clToConsider}`, 1)))
 						}
 					}
 					else {
-						changeToConsider.desc = await this.robo.p4.describe(clToConsider, 100)
+						changeToConsider.desc = await p4.describe(clToConsider, 100)
 					}
 				}
 			}
diff --git a/Engine/Extras/RoboMerge/v3/src/robo/nodebot.ts b/Engine/Extras/RoboMerge/v3/src/robo/nodebot.ts
--- a/Engine/Extras/RoboMerge/v3/src/robo/nodebot.ts
+++ b/Engine/Extras/RoboMerge/v3/src/robo/nodebot.ts
@@ -106,7 +106,7 @@ export class NodeBot extends PerforceStatefulBot implements NodeBotInterface {
 
 		this.externalRobomergeUrl = externalUrl
 
-		this.p4 = new PerforceContext(this.nodeBotLogger)
+		this.p4 = PerforceContext.getServerContext(this.nodeBotLogger)
 
 		this.mailer = mailer
 
diff --git a/Engine/Extras/RoboMerge/v3/src/robo/preview.ts b/Engine/Extras/RoboMerge/v3/src/robo/preview.ts
--- a/Engine/Extras/RoboMerge/v3/src/robo/preview.ts
+++ b/Engine/Extras/RoboMerge/v3/src/robo/preview.ts
@@ -15,7 +15,7 @@ export function init(_root: string) {
 export async function getPreview(cl: number, singleBot?: string) {
 	// @todo print content from shelf
 	if (!p4) {
-		p4 = new PerforceContext(logger)
+		p4 = PerforceContext.getServerContext(logger)
 	}
 
 	const status = new Status(new Date(), 'preview', logger)
diff --git a/Engine/Extras/RoboMerge/v3/src/robo/robo.ts b/Engine/Extras/RoboMerge/v3/src/robo/robo.ts
--- a/Engine/Extras/RoboMerge/v3/src/robo/robo.ts
+++ b/Engine/Extras/RoboMerge/v3/src/robo/robo.ts
@@ -10,7 +10,7 @@ import { Arg, readProcessArgs } from '../common/args';
 import { _setTimeout } from '../common/helper';
 import { ContextualLogger } from '../common/logger';
 import { Mailer } from '../common/mailer';
-import { getPerforceUsername, initializePerforce, PerforceContext } from '../common/perforce';
+import { initializePerforce, PerforceContext } from '../common/perforce';
 import { BuildVersion, VersionReader } from '../common/version';
 import { CertFiles } from '../common/webserver';
 import { addBranchGraph, Graph, GraphAPI } from '../new/graph';
@@ -233,20 +233,15 @@ if (env === 'prod'
 GraphBot.dataDirectory = args.branchSpecsDirectory;
 
 export class RoboMerge {
-	private readonly roboMergeLogger = new ContextualLogger('RoboMerge')
+	readonly roboMergeLogger = new ContextualLogger('RoboMerge')
 	readonly graphBots = new Map<string, GraphBot>()
 	graph: GraphAPI
 	mailer: Mailer
 	
 	static VERSION : BuildVersion = VersionReader.getBuildVersionObj();
-	readonly p4: PerforceContext
 
 	readonly VERSION_STRING = VersionReader.toString();
 
-	constructor() {
-		this.p4 = new PerforceContext(this.roboMergeLogger)
-	}
-
 	getAllBranches() : Branch[] {
 		const branches: Branch[] = [];
 		for (const graphBot of robo.graphBots.values()) {
@@ -333,13 +328,15 @@ async function _onBranchSpecReloaded(graphBot: GraphBot, logger: ContextualLogge
 
 async function init(logger: ContextualLogger) {
 
+	let p4 = PerforceContext.getServerContext(logger)
+
 	let lookupStream = async function(rootPath: string): Promise<string> {
 
 		while (true) {
 			try {
-				const stream = await robo.p4.getStreamName(rootPath);
+				const stream = await p4.getStreamName(rootPath);
 				if (typeof stream === 'string') {
-					if (await robo.p4.stream(stream)) {
+					if (await p4.stream(stream)) {
 						return stream
 					}
 				}
@@ -366,18 +363,18 @@ async function init(logger: ContextualLogger) {
 		}
 
 		// Ensure we have a workspace for branch specs
-		const workspace: Object[] = await robo.p4.find_workspace_by_name(args.branchSpecsWorkspace)
+		const workspace: Object[] = await p4.find_workspace_by_name(args.branchSpecsWorkspace)
 		if (workspace.length === 0) {
 			logger.info(`Cannot find workspace ${args.branchSpecsWorkspace}, creating a new one.`)
 
 			if ((await lookupStream(args.branchSpecsRootPath)).length > 0) {
-				await robo.p4.newBranchSpecWorkspace(autoUpdaterConfig.workspace, args.branchSpecsRootPath)
+				await p4.newBranchSpecWorkspace(autoUpdaterConfig.workspace, args.branchSpecsRootPath)
 			}
 		}
 
 		// make sure we've got the latest branch specs
 		logger.info('Syncing latest branch specs')
-		await AutoBranchUpdater.init({p4: robo.p4}, autoUpdaterConfig, logger)
+		await AutoBranchUpdater.init(p4, autoUpdaterConfig, robo.roboMergeLogger)
 
 	} else {
 		logger.warn('Auto brancher updater not configured!')
@@ -390,7 +387,7 @@ async function init(logger: ContextualLogger) {
 	}
 
 	if (args.persistenceBackupFrequency > 0 && !args.previewOnly) {
-		const workspace: Object[] = await robo.p4.find_workspace_by_name(args.persistenceBackupWorkspace)
+		const workspace: Object[] = await p4.find_workspace_by_name(args.persistenceBackupWorkspace)
 		if (workspace.length === 0) {
 			logger.info(`Cannot find workspace ${args.persistenceBackupWorkspace}, creating a new one.`)
 
@@ -405,7 +402,7 @@ async function init(logger: ContextualLogger) {
 					AltRoots: roots,
 					Stream: stream
 				}		
-				await robo.p4.newWorkspace(args.persistenceBackupWorkspace, params)
+				await p4.newWorkspace(args.persistenceBackupWorkspace, params)
 			}
 		}
 	} else {
@@ -568,8 +565,6 @@ async function main(logger: ContextualLogger) {
 			await _setTimeout(timeout*1000);
 		}
 	}
-	// log the user name
-	logger.info(`P4 appears to be configured. User=${getPerforceUsername()}`);
 
 	setGlobalAnalytics(new Analytics(args.hostInfo!))
 	robo = new RoboMerge;
