Rename ExecOpts edgeServerAddress to serverAddress

diff --git a/Engine/Extras/RoboMerge/v3/src/common/perforce.ts b/Engine/Extras/RoboMerge/v3/src/common/perforce.ts
--- a/Engine/Extras/RoboMerge/v3/src/common/perforce.ts
+++ b/Engine/Extras/RoboMerge/v3/src/common/perforce.ts
@@ -270,7 +270,7 @@ interface ExecOpts {
 	noCwd?: boolean
 	numRetries?: number
 	trace?: boolean
-	edgeServerAddress?: string
+	serverAddress?: string
 }
 
 interface ExecZtagOpts extends ExecOpts {
@@ -484,11 +484,11 @@ export class PerforceContext {
 		
 		const serverAddress = serverDetails && serverDetails.port
 
-		const output = await PerforceContext.execAndParse(logger, null, ["login", "-s"], { edgeServerAddress: serverAddress });
+		const output = await PerforceContext.execAndParse(logger, null, ["login", "-s"], { serverAddress });
 		let resp = output[0];
 	
 		if (resp && resp.User) {
-			const rawServerVersion = await PerforceContext._execP4(logger, null, ["-u", resp.User, "-ztag","-F","%serverVersion%","info"], { edgeServerAddress: serverAddress} )
+			const rawServerVersion = await PerforceContext._execP4(logger, null, ["-u", resp.User, "-ztag","-F","%serverVersion%","info"], { serverAddress} )
 			const match = rawServerVersion.match(/.*\/(\d+\.\d+)\/.*/)
 			if (!match) {
 				throw new Error(`Unable to parse server version from ${rawServerVersion}`)
@@ -498,7 +498,7 @@ export class PerforceContext {
 				throw new Error(`Robomerge requires a minimum server version of ${MIN_SERVER_VERSION}`)
 			}
 	
-			const serversOutput = await PerforceContext.execAndParse(logger, null, ["-u", resp.User, "servers"], { edgeServerAddress: serverAddress });
+			const serversOutput = await PerforceContext.execAndParse(logger, null, ["-u", resp.User, "servers"], { serverAddress });
 
 			let serverConfig: ServerConfig = {
 				username: resp.User,
@@ -551,7 +551,7 @@ export class PerforceContext {
 			args = [...args, '-c', workspaceName]
 		}
 
-		return this.execAndParse(null, args, { edgeServerAddress });
+		return this.execAndParse(null, args, { serverAddress: edgeServerAddress });
 	}
 
 	/** get a single change and return it in the format of changes() */
@@ -690,15 +690,13 @@ export class PerforceContext {
 
 		let args = ['clients', '-u', user || this.username]
 
-		// -a to include workspaces on edge servers
-		if (!edgeServerAddress) {
-			// find all
-			args.push('-a')
-		}
-
 		let opts: ExecOpts = {}
 		if (edgeServerAddress && edgeServerAddress !== 'commit') {
-			opts.edgeServerAddress = edgeServerAddress
+			opts.serverAddress = edgeServerAddress
+		}
+		else {
+			// -a to include workspaces on edge servers
+			args.push('-a')
 		}
 
 		let workspaces = [];
@@ -738,18 +736,12 @@ export class PerforceContext {
 
 		let args = ['clients', '-E', workspaceName]
 
-		// -a to include workspaces on edge servers
-		if (!edgeServerAddress) {
-			// find all
-			args.push('-a')
-		}
-
 		let opts: ExecOpts = {}
 		if (edgeServerAddress && edgeServerAddress !== 'commit') {
-			opts.edgeServerAddress = edgeServerAddress
+			opts.serverAddress = edgeServerAddress
 		}
-
-		if (edgeServerAddress) {
+		else {
+			// -a to include workspaces on edge servers
 			args.push('-a')
 		}
 
@@ -775,7 +767,7 @@ export class PerforceContext {
 	}
 
 	reloadWorkspace(workspaceName: string, edgeServerAddress?: string) {
-		return this.execAndParse(null, ['reload', '-c', workspaceName], {edgeServerAddress});
+		return this.execAndParse(null, ['reload', '-c', workspaceName], {serverAddress: edgeServerAddress});
 	}
 
 	getEdgeServerAddress(serverId: string) {
@@ -925,7 +917,7 @@ export class PerforceContext {
 		this.logger.info("Executing: 'p4 change -i' to create a new CL");
 		while (true) {
 			try {
-				const output = await this._execP4(workspace, ['change', '-i'], { stdin: form, quiet: true, edgeServerAddress });
+				const output = await this._execP4(workspace, ['change', '-i'], { stdin: form, quiet: true, serverAddress: edgeServerAddress });
 				// parse the CL out of output
 				const match = output.match(/Change (\d+) created./);
 				if (!match) {
@@ -1009,7 +1001,7 @@ export class PerforceContext {
 		// execute the P4 command
 		let changes;
 		try {
-			changes = await this.execAndParse(roboWorkspace, cmdList, { numRetries: 0, edgeServerAddress: opts.edgeServerAddress });
+			changes = await this.execAndParse(roboWorkspace, cmdList, { numRetries: 0, serverAddress: opts.edgeServerAddress });
 		}
 		catch (reason) {
 			if (!isExecP4Error(reason)) {
@@ -1071,7 +1063,7 @@ export class PerforceContext {
 	async resolveHelper(workspace: Workspace | null, flag: string, changelist: number, edgeServerAddress?: string): Promise<any> {
 		try {
 			// Perform merge
-			return await this._execP4Ztag(workspace, ['resolve', flag, `-c${changelist}`], {edgeServerAddress})
+			return await this._execP4Ztag(workspace, ['resolve', flag, `-c${changelist}`], {serverAddress: edgeServerAddress})
 		}
 		catch (reason) {
 			if (!isExecP4Error(reason)) {
@@ -1148,7 +1140,7 @@ export class PerforceContext {
 
 		let dashNresult: string[]
 		try {
-			dashNresult = await this._execP4Ztag(workspace, ['resolve', '-N', `-c${changelist}`], {edgeServerAddress, resolve: true})
+			dashNresult = await this._execP4Ztag(workspace, ['resolve', '-N', `-c${changelist}`], {serverAddress: edgeServerAddress, resolve: true})
 		}
 		catch (reason) {
 			if (!isExecP4Error(reason)) {
@@ -1174,7 +1166,7 @@ export class PerforceContext {
 	async submit(roboWorkspace: RoboWorkspace, changelist: number, edgeServerAddress?: string): Promise<number | string> {
 		let rawOutput: string
 		try {
-			rawOutput = await this._execP4(roboWorkspace, ['-ztag', 'submit', '-f', 'submitunchanged', '-c', changelist.toString()], {edgeServerAddress})
+			rawOutput = await this._execP4(roboWorkspace, ['-ztag', 'submit', '-f', 'submitunchanged', '-c', changelist.toString()], {serverAddress: edgeServerAddress})
 		}
 		catch ([errArg, output]) {
 			const err = errArg.toString().trim()
@@ -1253,7 +1245,7 @@ export class PerforceContext {
 	// delete a CL
 	// output format is just error or not
 	deleteCl(roboWorkspace: RoboWorkspace, changelist: number, edgeServerAddress?: string) {
-		return this._execP4(roboWorkspace, ["change", "-d", changelist.toString()], {edgeServerAddress});
+		return this._execP4(roboWorkspace, ["change", "-d", changelist.toString()], {serverAddress: edgeServerAddress});
 	}
 
 	// run p4 'opened' command: lists files in changelist with details of edit state (e.g. if a copy, provides source path)
@@ -1287,7 +1279,7 @@ export class PerforceContext {
 		const edgeServer = await this.getWorkspaceEdgeServer(client)
 		const args = ['revert', '-C', client, ...files]
 		try {
-			const results = await this.execAndParse(null, args, {edgeServerAddress: edgeServer?.address})
+			const results = await this.execAndParse(null, args, {serverAddress: edgeServer?.address})
 
 			// Files can fail to be reverted because of moves, lets try and resolve that
 			const movedFiles = 
@@ -1325,7 +1317,7 @@ export class PerforceContext {
 		try {
 			await this._execP4(workspace, ['revert',
 				/*"-w", seems to cause ENOENT occasionally*/ ...(opts || []), '-c', changelist.toString(), path],
-				{edgeServerAddress});
+				{serverAddress: edgeServerAddress});
 		}
 		catch (reason) {
 			if (!isExecP4Error(reason)) {
@@ -1372,7 +1364,7 @@ export class PerforceContext {
 	async shelve(roboWorkspace: RoboWorkspace, changelist: number, edgeServerAddress?: string) {
 		try {
 			await this._execP4(roboWorkspace, ['shelve', '-f', '-c', changelist.toString()],
-				{ numRetries: 0, edgeServerAddress })
+				{ numRetries: 0, serverAddress: edgeServerAddress })
 		}
 		catch (reason) {
 			if (!isExecP4Error(reason)) {
@@ -1490,7 +1482,7 @@ export class PerforceContext {
 		opts = opts || {}; // optional newWorkspace:string and/or changeSubmitted:boolean
 		// get the current changelist description
 		let form = await this._execP4(workspace, ['change', '-o', changelist.toString()],
-			{edgeServerAddress: opts.edgeServerAddress});
+			{serverAddress: opts.edgeServerAddress});
 
 		if (opts.newOwner) {
 			form = form.replace(/\nUser:\t[^\n]*\n/, `\nUser:\t${opts.newOwner}\n`);
@@ -1509,7 +1501,7 @@ export class PerforceContext {
 		this.logger.info(`Executing: 'p4 change -i ${changeFlag}' to edit CL ${changelist}`);
 
 		await this._execP4(workspace, ['change', '-i', changeFlag],
-				{ stdin: form, quiet: true, edgeServerAddress: opts.edgeServerAddress })
+				{ stdin: form, quiet: true, serverAddress: opts.edgeServerAddress })
 	}
 
 	// change the description on an existing CL
@@ -1634,8 +1626,8 @@ export class PerforceContext {
 			args = ['-c', workspace.name, ...args]
 		}
 
-		if (opts.edgeServerAddress) {
-			args = ['-p', opts.edgeServerAddress, ...args]
+		if (opts.serverAddress) {
+			args = ['-p', opts.serverAddress, ...args]
 		}
 
 		args = ['-zprog=robomerge', '-zversion=' + robomergeVersion, ...args]
