Use edge specific logger rather than a generic p4utils logger
Remove unused code

diff --git a/Engine/Extras/RoboMerge/v3/src/common/p4util.ts b/Engine/Extras/RoboMerge/v3/src/common/p4util.ts
--- a/Engine/Extras/RoboMerge/v3/src/common/p4util.ts
+++ b/Engine/Extras/RoboMerge/v3/src/common/p4util.ts
@@ -1,180 +1,14 @@
 // Copyright Epic Games, Inc. All Rights Reserved.
 
 import { ContextualLogger } from './logger';
-import { Change, ClientSpec, coercePerforceWorkspace, OpenedFileRecord, PerforceContext, RoboWorkspace } from './perforce';
+import { Change, ClientSpec, PerforceContext } from './perforce';
 
 const USER_WORKSPACE_EXCLUDE_PATTERNS: (RegExp | string)[] = [
 	'horde-p4bridge-',
 	'swarm-'
 ]
 
-///////////////////
-// Error reporting
-/*function _reportFilesStillToResolve(changeNum: number, needsResolve: Object[]) {
-	// @todo list files in error
-	throw new Error(`changelist ${changeNum} has files still to resolve`)
-}*/
-
-const p4utilsLogger = new ContextualLogger('P4Utils')
-
-function _reportUnknownActions(changeNum: number, unknownActions: OpenedFileRecord[]) {
-	let errText = 'ERROR: These files have unknown actions so cannot be reliably re-opened:\n'
-	for (const openedRec of unknownActions) {
-		errText += `    ${openedRec.action},${openedRec.clientFile}\n`
-	}
-	p4utilsLogger.error(errText.trim())
-	throw new Error(`changelist ${changeNum} has files with unknown action type (see log)`)
-}
-
-/*function _warnNotSyncedOpenFiles(notSyncedOpenFiles: OpenedFileRecord[]) {
-	let warnText = 'WARNING: the following files were not synced to the revision they were resolved to before. syncing them to the indicated rev first:\n'
-	for (const file of notSyncedOpenFiles) {
-		warnText += `    ${file.clientFile}#${file.rev} (had #${file.haveRev})\n`
-	}
-	p4utilsLogger.warn(warnText.trim())
-}*/
-///////////////////
-
-export async function convertIntegrateToEdit(p4: PerforceContext, roboWorkspace: RoboWorkspace, changeNum: number) {
-	const workspace = coercePerforceWorkspace(roboWorkspace)!
-	const files = await p4.opened(workspace, changeNum) as OpenedFileRecord[]
-	if (files.length === 0) {
-		throw new Error('nothing to integrate')
-	}
-
-	// @todo check all files have all OpenedFileRecord fields
-
-	p4utilsLogger.verbose(`Converting ${files.length} file${files.length === 1 ? '' : 's'} in CL ${changeNum}.`)
-
-	// create a mapping of client to local files so we can re-open using local paths (needed for some actions)
-	const openedClientToLocalMap = new Map(files.map(file => [file.clientFile.toLowerCase(), file.depotFile] as [string, string]))
-	const needResolve = await p4.listFilesToResolve(workspace, changeNum)
-	if (needResolve.length !== 0) {
-		// @todo list files in error
-		throw new Error(`changelist ${changeNum} has files still to resolve`)
-	}
-
-	// We'll just assume all the files in the changelist are integrated. We'll revert and re-open
-	// with the same action/filetype. This should remove the integration records.
-
-	// partition our files into ones we know the actions for, and the ones we don't
-	const basicActions = new Map<string, OpenedFileRecord[]>([
-		['delete', []],
-		['add', []],
-		['edit', []],
-	])
-	const moveAddActions: OpenedFileRecord[] = []
-	const moveDeleteActions: OpenedFileRecord[] = []
-	const unknownActions: OpenedFileRecord[] = []
-
-	// multiple action map to the same reopen mapping, so this handles that.
-	const actionMappings = new Map<string, OpenedFileRecord[]>([
-		['delete', basicActions.get('delete')!],
-		['branch', basicActions.get('add')!],
-		['add', basicActions.get('add')!],
-		['edit', basicActions.get('edit')!],
-		['integrate', basicActions.get('edit')!],
-		['move/add', moveAddActions],
-		['move/delete', moveDeleteActions],
-	])
-
-	// this maps all open records to a reopen action.
-	for (const openedRec of files) {
-		(actionMappings.get(openedRec.action) || unknownActions).push(openedRec)
-	}
-
-	// if we have an unknown action, abort with error.
-	if (unknownActions.length !== 0) {
-		_reportUnknownActions(changeNum, unknownActions)
-	}
-
-	// we are good to go, so actually revert the files.
-	p4utilsLogger.verbose(`Reverting files that will be reopened`)
-
-	await p4.revert(workspace, changeNum, ['-k'])
-
-	// P4 lets you resolve/submit integrations without syncing them to the revision you are resolving against because it's a server operation.
-	// if the haveRev doesn't match the rev of the originally opened file, this was the case.
-	// so we sync -k because the local file matches the one we want, we just have to tell P4 that so it will let us check it out at that revision.
-	// If we simply sync -k to #head we might miss a legitimate submit that happened since our last resolve, and that would stomp that submit.
-	const notSyncedOpenFiles: OpenedFileRecord[] = files.filter(file => 
-		file.haveRev !== file.rev && file.action !== 'branch' && file.action !== 'add')
-
-	if (notSyncedOpenFiles.length !== 0) {
-// don't warn - this will almost always be the case for RoboMerge
-//		_warnNotSyncedOpenFiles(notSyncedOpenFiles);
-		await Promise.all(notSyncedOpenFiles.map(file => p4.sync(workspace, `${file.clientFile}#${file.rev}`, {opts: ['-k']})))
-	}
-
-	// Perform basic actions
-	for (const [action, bucket] of basicActions.entries()) {
-		if (bucket.length !== 0) {
-			const localFiles = bucket.map(file => openedClientToLocalMap.get(file.clientFile.toLowerCase())!)
-			let reopenText = `Re-opening the following files for ${action}:`
-			for (const file of localFiles) {
-				reopenText += `    ${file}\n`
-			}
-			p4utilsLogger.verbose(reopenText.trim())
-			await p4.run(workspace, action, changeNum, localFiles)
-		}
-	}
-
-	// Perform move actions
-	if (moveAddActions.length !== 0) {
-		const localFiles = moveAddActions.map(file => [
-			file.movedFile!,
-			openedClientToLocalMap.get(file.clientFile.toLowerCase())!
-		])
-
-		// todo: check that movedFile actually exists
-		p4utilsLogger.verbose(`Re-opening the following files for move:`)
-
-		for (const [src, target] of localFiles) {
-			if (!src) {
-				throw new Error(`No source file in move (CL ${changeNum})`)
-			}
-
-			p4utilsLogger.verbose(`    ${src} to ${target}`)
-		}
-
-		// We have to first open the source file for edit (have to use -k because the file has already been moved locally!)
-		await p4.run(workspace, 'edit', changeNum, localFiles.map(([src, _target]) => src), ['-k'])
-
-		// then we can open the file for move in the new location (have to use -k because the file has already been moved locally!)
-		for (const [src, target] of localFiles) {
-			await p4.move(workspace, changeNum, src, target, ['-k'])
-		}
-	}
-
-	// Get the list of reopened files in the CL to check their filetype
-	p4utilsLogger.verbose(`Getting the list of files reopened in changelist ${changeNum}...`)
-
-	const reopened = await p4.opened(workspace, changeNum) as OpenedFileRecord[]
-
-	if (reopened.length === 0) {
-		throw new Error('change has no reopened files. This is an error in the conversion code')
-	}
-
-	if (reopened.length !== files.length) {
-		// should block stream and get human to resolve?
-		throw new Error(`change doesn't have the same number of reopened files (${reopened.length}) as originally (${files.length}). ` +
-			'This probably signifies an error in the conversion code, and the actions should be reverted!')
-	}
-
-	const reopenedRecordsMap = new Map(reopened.map(x => [x.clientFile, x] as [string, OpenedFileRecord]))
-	for (const file of files) {
-		const reopenedRecord = reopenedRecordsMap.get(file.clientFile)
-		if (!reopenedRecord) {
-			throw new Error(`ERROR: Could not find original file ${file.clientFile} in re-opened records. This signifies an error in the conversion code! Aborting...`)
-		}
-		if (file.type !== reopenedRecord.type) {
-			p4utilsLogger.verbose(`Changing filetype of ${file.clientFile} from ${reopenedRecord.type} to ${file.type}.`)
-			await p4.run(workspace, 'reopen', changeNum, [reopenedRecord.clientFile], ['t', file.type])
-		}
-	}
-}
-
-export async function cleanWorkspace(p4: PerforceContext, workspaceName: string, workspaceRoot: string, edgeServerAddress?: string) {
+export async function cleanWorkspace(logger: ContextualLogger, p4: PerforceContext, workspaceName: string, workspaceRoot: string, edgeServerAddress?: string) {
 	const changes = await p4.get_pending_changes(workspaceName) as Change[]
 
 	// revert any pending changes left over from previous runs
@@ -182,28 +16,28 @@ export async function cleanWorkspace(p4: PerforceContext, workspaceName: string,
 		const changeStr = `CL ${change.change}: ${change.desc}`
 		const workspace = change.client
 		if (change.shelved) {
-			p4utilsLogger.info(`Attempting to delete shelved files in ${changeStr}`)
+			logger.info(`Attempting to delete shelved files in ${changeStr}`)
 			try {
 				await p4.delete_shelved(workspace, change.change)
 			}
 			catch (err) {
 				// ignore delete errors on startup (as long as delete works, we're good)
-				p4utilsLogger.error(`Delete shelved failed. Will try revert anyway: ${err}`)
+				logger.error(`Delete shelved failed. Will try revert anyway: ${err}`)
 			}
 		}
-		p4utilsLogger.info(`Attempting to revert ${changeStr}`)
+		logger.info(`Attempting to revert ${changeStr}`)
 		try {
 			await p4.revert(workspace, change.change, [], edgeServerAddress)
 		}
 		catch (err) {
 			// ignore revert errors on startup (As long as delete works, we're good)
-			p4utilsLogger.error(`Revert failed. Will try delete anyway: ${err}`)
+			logger.error(`Revert failed. Will try delete anyway: ${err}`)
 		}
 
 		await p4.deleteCl(workspace, change.change, edgeServerAddress)
 	}
 
-	p4utilsLogger.info('Resetting all workspaces to revision 0')
+	logger.info('Resetting all workspaces to revision 0')
 	await p4.sync(workspaceName, workspaceRoot + '#0', {edgeServerAddress})
 }
 
diff --git a/Engine/Extras/RoboMerge/v3/src/robo/autobranchupdater.ts b/Engine/Extras/RoboMerge/v3/src/robo/autobranchupdater.ts
--- a/Engine/Extras/RoboMerge/v3/src/robo/autobranchupdater.ts
+++ b/Engine/Extras/RoboMerge/v3/src/robo/autobranchupdater.ts
@@ -191,7 +191,7 @@ export class AutoBranchUpdater implements Bot {
 
 					if (retryCount < 5) {
 						this.abuLogger.warn(`Mirror file reloading error (retries ${retryCount}: ${err}`)
-						p4util.cleanWorkspace(this.p4, mirrorWorkspace.name, mirrorWorkspace.depotpath)
+						p4util.cleanWorkspace(this.abuLogger, this.p4, mirrorWorkspace.name, mirrorWorkspace.depotpath)
 					}
 					else {
 						this.abuLogger.error('Mirror file reloading failed, retries exhausted')
diff --git a/Engine/Extras/RoboMerge/v3/src/robo/edgebot.ts b/Engine/Extras/RoboMerge/v3/src/robo/edgebot.ts
--- a/Engine/Extras/RoboMerge/v3/src/robo/edgebot.ts
+++ b/Engine/Extras/RoboMerge/v3/src/robo/edgebot.ts
@@ -492,7 +492,7 @@ class EdgeBotImpl extends PerforceStatefulBot {
 			if (existingWorkspaceInfo[0].IsUnloaded) {
 				await this.p4.reloadWorkspace(workspaceName, edgeServerAddress)
 			}
-			await p4util.cleanWorkspace(this.p4, workspaceName, this.targetBranch.rootPath, edgeServerAddress)
+			await p4util.cleanWorkspace(this.logger, this.p4, workspaceName, this.targetBranch.rootPath, edgeServerAddress)
 		}
 		else {
 			const params: any = {};
