Prevent workspace from cleaned twice, potentially while being used

diff --git a/Engine/Extras/RoboMerge/v3/src/common/p4util.ts b/Engine/Extras/RoboMerge/v3/src/common/p4util.ts
--- a/Engine/Extras/RoboMerge/v3/src/common/p4util.ts
+++ b/Engine/Extras/RoboMerge/v3/src/common/p4util.ts
@@ -1,5 +1,6 @@
 // Copyright Epic Games, Inc. All Rights Reserved.
 
+import { _setTimeout } from './helper';
 import { ContextualLogger } from './logger';
 import { Change, ClientSpec, PerforceContext } from './perforce';
 
@@ -8,7 +9,25 @@ const USER_WORKSPACE_EXCLUDE_PATTERNS: (RegExp | string)[] = [
 	'swarm-'
 ]
 
+// true has been cleaned, false is being cleaned, undefined needs cleaning
+let cleanedWorkspaces: Map<string, boolean> = new Map()
+
 export async function cleanWorkspace(logger: ContextualLogger, p4: PerforceContext, workspaceName: string, workspaceRoot: string, edgeServerAddress?: string) {
+	
+	// if the workspace is being cleaned, wait for it to finish
+	let workspaceCleaned = cleanedWorkspaces.get(workspaceName)
+	while (workspaceCleaned === false) {
+		logger.info(`Waiting on ${workspaceName} being cleaned`)
+		await _setTimeout(1000)
+		workspaceCleaned = cleanedWorkspaces.get(workspaceName)
+	}
+
+	if (workspaceCleaned) {
+		return
+	}
+
+	logger.info(`Cleaning ${workspaceName}`)
+	cleanedWorkspaces.set(workspaceName, false)
 	const changes = await p4.get_pending_changes(workspaceName) as Change[]
 
 	// revert any pending changes left over from previous runs
@@ -37,8 +56,10 @@ export async function cleanWorkspace(logger: ContextualLogger, p4: PerforceConte
 		await p4.deleteCl(workspace, change.change, edgeServerAddress)
 	}
 
-	logger.info('Resetting all workspaces to revision 0')
+	logger.info(`Resetting workspace ${workspaceName} to revision 0`)
 	await p4.sync(workspaceName, workspaceRoot + '#0', {edgeServerAddress})
+
+	cleanedWorkspaces.set(workspaceName, true)
 }
 
 export async function getWorkspacesForUser(p4: PerforceContext, user: string) {
diff --git a/Engine/Extras/RoboMerge/v3/src/robo/autobranchupdater.ts b/Engine/Extras/RoboMerge/v3/src/robo/autobranchupdater.ts
--- a/Engine/Extras/RoboMerge/v3/src/robo/autobranchupdater.ts
+++ b/Engine/Extras/RoboMerge/v3/src/robo/autobranchupdater.ts
@@ -191,7 +191,7 @@ export class AutoBranchUpdater implements Bot {
 
 					if (retryCount < 5) {
 						this.abuLogger.warn(`Mirror file reloading error (retries ${retryCount}: ${err}`)
-						p4util.cleanWorkspace(this.abuLogger, this.p4, mirrorWorkspace.name, mirrorWorkspace.depotpath)
+						await p4util.cleanWorkspace(this.abuLogger, this.p4, mirrorWorkspace.name, mirrorWorkspace.depotpath)
 					}
 					else {
 						this.abuLogger.error('Mirror file reloading failed, retries exhausted')
