Mirror fixes
* File does not need to already exist
* File can go in root folder of a stream
* Sync/edit are done virtually
* No longer output #jira none in the changelist description

diff --git a/Engine/Extras/RoboMerge/v3/src/robo/autobranchupdater.ts b/Engine/Extras/RoboMerge/v3/src/robo/autobranchupdater.ts
--- a/Engine/Extras/RoboMerge/v3/src/robo/autobranchupdater.ts
+++ b/Engine/Extras/RoboMerge/v3/src/robo/autobranchupdater.ts
@@ -239,10 +239,16 @@ export class AutoBranchUpdater implements Bot {
 
 		const {depotpath, realFilepath, mirrorFilepath} = workspace
 
-		await this.p4.sync(workspace, depotpath, {opts:[P4_FORCE]})
-		const cl = await this.p4.new_cl(workspace, "Updating mirror file\n#jira none\n")
-		await this.p4.edit(workspace, cl, depotpath)
+		const cl = await this.p4.new_cl(workspace, "Updating mirror file\n")
+		const fstat = await this.p4.fstat(workspace, depotpath)
+		if (fstat.length > 0) {
+			await this.p4.sync(workspace, depotpath, {opts:['-k']})
+			await this.p4.edit(workspace, cl, depotpath, ['-k'])
+		}
 		await new Promise((done, _) => fs.copyFile(realFilepath, mirrorFilepath, done))
+		if (fstat.length == 0) {
+			await this.p4.add(workspace, cl, depotpath)
+		}
 		await this.p4.submit(workspace, cl)
 	}
 
@@ -256,9 +262,13 @@ export class AutoBranchUpdater implements Bot {
 			throw new Error('If mirrorPath is specified, expecting [stream, sub-path, filename]')
 		}
 
+		if (mirrorPathBits[0].length == 0 || mirrorPathBits[2].length == 0) {
+			throw new Error('stream and filename must be supplied for mirrorPath')
+		}
+
 		const directory = AutoBranchUpdater.config!.workspace.directory
-		const depotFolder = mirrorPathBits.slice(0, 2).join('/')
-		const depotpath = mirrorPathBits.join('/')
+		const depotFolder = mirrorPathBits[1].length > 0 ? mirrorPathBits.slice(0, 2).join('/') : mirrorPathBits[0]
+		const depotpath = `${depotFolder}/${mirrorPathBits[2]}`
 		const realFilepath = path.join(directory, bot.filename)
 		const mirrorFilepath = path.join(directory, ...mirrorPathBits.slice(1))
 		return {
