Ensure that directory has been created for mirror files to be written to the workspace

diff --git a/Engine/Extras/RoboMerge/v3/src/robo/autobranchupdater.ts b/Engine/Extras/RoboMerge/v3/src/robo/autobranchupdater.ts
--- a/Engine/Extras/RoboMerge/v3/src/robo/autobranchupdater.ts
+++ b/Engine/Extras/RoboMerge/v3/src/robo/autobranchupdater.ts
@@ -245,7 +245,8 @@ export class AutoBranchUpdater implements Bot {
 			await this.p4.sync(workspace, depotpath, {opts:['-k']})
 			await this.p4.edit(workspace, cl, depotpath, ['-k'])
 		}
-		await new Promise((done, _) => fs.copyFile(realFilepath, mirrorFilepath, done))
+		await fs.mkdirSync(path.dirname(mirrorFilepath), { recursive: true })
+		await fs.copyFileSync(realFilepath, mirrorFilepath)
 		if (fstat.length == 0) {
 			await this.p4.add(workspace, cl, depotpath)
 		}
