Clean mirror workspace before trying to do mirror operations
Fixed clean workspace to work correctly when workspace name and root directory don't match and root directory doesn't exist on disk
Added logging to try and diagnose failure to mirror in production

diff --git a/Engine/Extras/RoboMerge/v3/src/common/p4util.ts b/Engine/Extras/RoboMerge/v3/src/common/p4util.ts
--- a/Engine/Extras/RoboMerge/v3/src/common/p4util.ts
+++ b/Engine/Extras/RoboMerge/v3/src/common/p4util.ts
@@ -2,7 +2,7 @@
 
 import { _setTimeout } from './helper';
 import { ContextualLogger } from './logger';
-import { Change, ClientSpec, PerforceContext } from './perforce';
+import { Change, ClientSpec, PerforceContext, Workspace } from './perforce';
 
 const USER_WORKSPACE_EXCLUDE_PATTERNS: (RegExp | string)[] = [
 	'horde-p4bridge-',
@@ -12,28 +12,27 @@ const USER_WORKSPACE_EXCLUDE_PATTERNS: (RegExp | string)[] = [
 // true has been cleaned, false is being cleaned, undefined needs cleaning
 let cleanedWorkspaces: Map<string, boolean> = new Map()
 
-export async function cleanWorkspace(logger: ContextualLogger, p4: PerforceContext, workspaceName: string, workspaceRoot: string, edgeServerAddress?: string) {
+export async function cleanWorkspace(logger: ContextualLogger, p4: PerforceContext, workspace: Workspace, workspaceRoot: string, edgeServerAddress?: string) {
 	
 	// if the workspace is being cleaned, wait for it to finish
-	let workspaceCleaned = cleanedWorkspaces.get(workspaceName)
+	let workspaceCleaned = cleanedWorkspaces.get(workspace.name)
 	while (workspaceCleaned === false) {
-		logger.info(`Waiting on ${workspaceName} being cleaned`)
+		logger.info(`Waiting on ${workspace.name} being cleaned`)
 		await _setTimeout(1000)
-		workspaceCleaned = cleanedWorkspaces.get(workspaceName)
+		workspaceCleaned = cleanedWorkspaces.get(workspace.name)
 	}
 
 	if (workspaceCleaned) {
 		return
 	}
 
-	logger.info(`Cleaning ${workspaceName}`)
-	cleanedWorkspaces.set(workspaceName, false)
-	const changes = await p4.get_pending_changes(workspaceName) as Change[]
+	logger.info(`Cleaning ${workspace.name}`)
+	cleanedWorkspaces.set(workspace.name, false)
+	const changes = await p4.get_pending_changes(workspace.name) as Change[]
 
 	// revert any pending changes left over from previous runs
 	for (const change of changes) {
 		const changeStr = `CL ${change.change}: ${change.desc}`
-		const workspace = change.client
 		if (change.shelved) {
 			logger.info(`Attempting to delete shelved files in ${changeStr}`)
 			try {
@@ -56,10 +55,10 @@ export async function cleanWorkspace(logger: ContextualLogger, p4: PerforceConte
 		await p4.deleteCl(workspace, change.change, edgeServerAddress)
 	}
 
-	logger.info(`Resetting workspace ${workspaceName} to revision 0`)
-	await p4.sync(workspaceName, workspaceRoot + '#0', {edgeServerAddress})
+	logger.info(`Resetting workspace ${workspace.name} to revision 0`)
+	await p4.sync(workspace, workspaceRoot + '#0', {edgeServerAddress})
 
-	cleanedWorkspaces.set(workspaceName, true)
+	cleanedWorkspaces.set(workspace.name, true)
 }
 
 export async function getWorkspacesForUser(p4: PerforceContext, user: string) {
diff --git a/Engine/Extras/RoboMerge/v3/src/robo/autobranchupdater.ts b/Engine/Extras/RoboMerge/v3/src/robo/autobranchupdater.ts
--- a/Engine/Extras/RoboMerge/v3/src/robo/autobranchupdater.ts
+++ b/Engine/Extras/RoboMerge/v3/src/robo/autobranchupdater.ts
@@ -182,23 +182,7 @@ export class AutoBranchUpdater implements Bot {
 				return
 			}
 
-			for (let retryCount = 0; ; ++retryCount) {
-				try {
-					await this.updateMirror(mirrorWorkspace)
-					return
-				}
-				catch (err) {
-
-					if (retryCount < 5) {
-						this.abuLogger.warn(`Mirror file reloading error (retries ${retryCount}: ${err}`)
-						await p4util.cleanWorkspace(this.abuLogger, this.p4, mirrorWorkspace.name, mirrorWorkspace.depotpath)
-					}
-					else {
-						this.abuLogger.error('Mirror file reloading failed, retries exhausted')
-						throw err
-					}
-				}
-			}
+			await this.updateMirror(mirrorWorkspace)
 		})
 	}
 
@@ -226,6 +210,9 @@ export class AutoBranchUpdater implements Bot {
 
 		const stream = this.graphBot.branchGraph.config.mirrorPath[0]
 
+		// make sure the root directory exists
+		fs.mkdirSync(workspace.directory, {recursive:true})
+
 		const workspaceQueryResult = await this.p4.find_workspace_by_name(workspace.name, {includeUnloaded: true})
 		if (workspaceQueryResult.length === 0) {
 			await this.p4.newWorkspace(workspace.name, {
@@ -233,8 +220,11 @@ export class AutoBranchUpdater implements Bot {
 				Root: AutoBranchUpdater.config!.workspace.directory
 			})
 		}
-		else if (workspaceQueryResult[0].IsUnloaded) {
-			await this.p4.reloadWorkspace(workspace.name)
+		else {
+			if (workspaceQueryResult[0].IsUnloaded) {
+				await this.p4.reloadWorkspace(workspace.name)
+			}
+			await p4util.cleanWorkspace(this.abuLogger, this.p4, workspace, workspace.depotpath)
 		}
 
 		const {depotpath, realFilepath, mirrorFilepath} = workspace
@@ -242,14 +232,26 @@ export class AutoBranchUpdater implements Bot {
 		const cl = await this.p4.new_cl(workspace, "Updating mirror file\n")
 		const fstat = await this.p4.fstat(workspace, depotpath)
 		if (fstat.length > 0) {
-			await this.p4.sync(workspace, depotpath, {opts:['-k']})
-			await this.p4.edit(workspace, cl, depotpath, ['-k'])
+			await this.p4.sync(workspace, depotpath, {opts:[P4_FORCE]})
+			await this.p4.edit(workspace, cl, depotpath)
+		}
+		else {
+			fs.mkdirSync(path.dirname(mirrorFilepath), { recursive: true })
+			if (!fs.existsSync(path.dirname(mirrorFilepath))) {
+				this.abuLogger.info(`Failed to create directory ${path.dirname(mirrorFilepath)}`)
+			}
+		}
+		if (!fs.existsSync(realFilepath)) {
+			this.abuLogger.info(`${realFilepath} does not exist to copy from`)
+		}
+		 fs.copyFileSync(realFilepath, mirrorFilepath)
+		if (!fs.existsSync(mirrorFilepath)) {
+			this.abuLogger.info(`Failed to copy to file ${mirrorFilepath}`)
 		}
-		await fs.mkdirSync(path.dirname(mirrorFilepath), { recursive: true })
-		await fs.copyFileSync(realFilepath, mirrorFilepath)
 		if (fstat.length == 0) {
 			await this.p4.add(workspace, cl, depotpath)
 		}
+		this.abuLogger.info(JSON.stringify(await this.p4.describe(cl)))
 		await this.p4.submit(workspace, cl)
 	}
 
diff --git a/Engine/Extras/RoboMerge/v3/src/robo/edgebot.ts b/Engine/Extras/RoboMerge/v3/src/robo/edgebot.ts
--- a/Engine/Extras/RoboMerge/v3/src/robo/edgebot.ts
+++ b/Engine/Extras/RoboMerge/v3/src/robo/edgebot.ts
@@ -492,7 +492,7 @@ class EdgeBotImpl extends PerforceStatefulBot {
 			if (existingWorkspaceInfo[0].IsUnloaded) {
 				await this.p4.reloadWorkspace(workspaceName, edgeServerAddress)
 			}
-			await p4util.cleanWorkspace(this.logger, this.p4, workspaceName, this.targetBranch.rootPath, edgeServerAddress)
+			await p4util.cleanWorkspace(this.logger, this.p4, {name: workspaceName, directory: path}, this.targetBranch.rootPath, edgeServerAddress)
 		}
 		else {
 			const params: any = {};
