Remove mirror debugging

diff --git a/Engine/Extras/RoboMerge/v3/src/robo/autobranchupdater.ts b/Engine/Extras/RoboMerge/v3/src/robo/autobranchupdater.ts
--- a/Engine/Extras/RoboMerge/v3/src/robo/autobranchupdater.ts
+++ b/Engine/Extras/RoboMerge/v3/src/robo/autobranchupdater.ts
@@ -237,21 +237,12 @@ export class AutoBranchUpdater implements Bot {
 		}
 		else {
 			fs.mkdirSync(path.dirname(mirrorFilepath), { recursive: true })
-			if (!fs.existsSync(path.dirname(mirrorFilepath))) {
-				this.abuLogger.info(`Failed to create directory ${path.dirname(mirrorFilepath)}`)
-			}
-		}
-		if (!fs.existsSync(realFilepath)) {
-			this.abuLogger.info(`${realFilepath} does not exist to copy from`)
 		}
 		 fs.copyFileSync(realFilepath, mirrorFilepath)
-		if (!fs.existsSync(mirrorFilepath)) {
-			this.abuLogger.info(`Failed to copy to file ${mirrorFilepath}`)
-		}
 		if (fstat.length == 0) {
 			await this.p4.add(workspace, cl, depotpath)
 		}
-		this.abuLogger.info(JSON.stringify(await this.p4.describe(cl)))
+		JSON.stringify(await this.p4.describe(cl))
 		await this.p4.submit(workspace, cl)
 	}
 
