Fix partial match message logic being inverted

diff --git a/Engine/Extras/RoboMerge/v3/src/robo/nodebot.ts b/Engine/Extras/RoboMerge/v3/src/robo/nodebot.ts
--- a/Engine/Extras/RoboMerge/v3/src/robo/nodebot.ts
+++ b/Engine/Extras/RoboMerge/v3/src/robo/nodebot.ts
@@ -1614,7 +1614,7 @@ export class NodeBot extends PerforceStatefulBot implements NodeBotInterface {
 				if (targetWorkspaceDef) {
 					result.info.targetWorkspaceForShelf = targetWorkspaceDef.client
 					if (targetWorkspaceDef.Stream) {
-						result.info.targetWorkspaceIsPartialMatch = targetWorkspaceDef.Stream.toLowerCase() == targetStream
+						result.info.targetWorkspaceIsPartialMatch = targetWorkspaceDef.Stream.toLowerCase() != targetStream
 					}
 				}
 				optWorkspaceOverride = result.info.targetWorkspaceForShelf
