Convert robomerge to using bolt SDK for slack connectivity

diff --git a/Engine/Extras/RoboMerge/v3/functional_tests/framework/src/framework.ts b/Engine/Extras/RoboMerge/v3/functional_tests/framework/src/framework.ts
--- a/Engine/Extras/RoboMerge/v3/functional_tests/framework/src/framework.ts
+++ b/Engine/Extras/RoboMerge/v3/functional_tests/framework/src/framework.ts
@@ -13,6 +13,15 @@ class DummyThread {
 	messages: string[] // the first message is the original post, the rest are replies
 }
 
+// This must be kept in sync with ipc.ts
+type OperationReturnType = {
+	statusCode: number
+	message: string // Goal: to provide meaningful error messaging to the end user
+
+	// Open Ended return
+	data?: any
+}
+
 const getJson = bent('json')
 
 const colors = require('colors')
@@ -23,7 +32,6 @@ colors.setTheme(require('colors/themes/generic-logging.js'))
 export const WORKSPACES_ROOT = process.env.WORKSPACES_ROOT || '/rm_tests/'
 const ROBOMERGE_DOMAIN_NAME = process.env.ROBOMERGE_DOMAIN || 'robomerge_functtest'
 export const ROBOMERGE_DOMAIN = `http://${ROBOMERGE_DOMAIN_NAME}:8877`
-export const DUMMY_SLACK_DOMAIN = `http://${ROBOMERGE_DOMAIN_NAME}:8811`
 
 // export const BOT_NAME = 'FUNCTIONALTEST'
 export const OPERATION_URL_TEMPLATE = ROBOMERGE_DOMAIN + `/api/op/bot/<bot>/node/<node>/op/<op>`
@@ -737,11 +745,11 @@ export abstract class FunctionalTest {
 
 	wasMessagePostedToSlack(channel: string, cl: number, target: string) {
 		const targetName = this.fullBranchName(target).toUpperCase()
-		return getJson(DUMMY_SLACK_DOMAIN + '/posted/' + channel)
-			.then((perTargetThreads: [string, DummyThread[]][]) => {
-				for (const [target, dummyThreads] of perTargetThreads) {
+		return getJson(ROBOMERGE_DOMAIN + '/dummyslack/' + channel)
+			.then((perTargetThreads: OperationReturnType ) => {
+				for (const [target, dummyThreads] of perTargetThreads.data) {
 					if (target === targetName) {
-						return dummyThreads.filter(dt => dt.CL == cl).length == 1
+						return dummyThreads.filter((dt: DummyThread) => dt.CL == cl).length == 1
 					}
 				}
 				return false
diff --git a/Engine/Extras/RoboMerge/v3/local-docker-run.sh b/Engine/Extras/RoboMerge/v3/local-docker-run.sh
--- a/Engine/Extras/RoboMerge/v3/local-docker-run.sh
+++ b/Engine/Extras/RoboMerge/v3/local-docker-run.sh
@@ -17,7 +17,6 @@ export ROBO_LOG_LEVEL=info
 export P4PORT=perforce:1666
 export P4PASSWD=
 export BOTS=robomergeQA1
-export ROBO_EXTERNAL_URL=http://localhost:8877
 
 set -x
 
diff --git a/Engine/Extras/RoboMerge/v3/package.json b/Engine/Extras/RoboMerge/v3/package.json
--- a/Engine/Extras/RoboMerge/v3/package.json
+++ b/Engine/Extras/RoboMerge/v3/package.json
@@ -16,6 +16,7 @@
     "@hpcc-js/graph": "^2.27.0",
     "@hpcc-js/wasm": "^0.3.14",
     "@sentry/node": "^5.7.1",
+    "@slack/bolt": "^3.22.0",
     "@types/ldapjs": "^1.0.3",
     "@types/marked": "^4.0.2",
     "@types/mustache": "^0.8.32",
diff --git a/Engine/Extras/RoboMerge/v3/run-tests.sh b/Engine/Extras/RoboMerge/v3/run-tests.sh
--- a/Engine/Extras/RoboMerge/v3/run-tests.sh
+++ b/Engine/Extras/RoboMerge/v3/run-tests.sh
@@ -82,7 +82,6 @@ docker run -d -p 8877:8877 -p 8811:8811 -p 25:25 -h robomerge_functtest --name r
     -e ROBO_EXTERNAL_URL=http://robomerge_functtest:8877  \
     -e BOTNAME=ft1,ft2,ft3,ft4,targets  \
     -e ROBO_LOG_LEVEL=info \
-    -e ROBO_SLACK_DOMAIN=http://localhost:8811 \
     robomerge node --trace-warnings dist/robo/watchdog.js
 
 docker run -t --hostname robomerge_functionaltests --name robomerge_functionaltests --network robomerge_functtest_network  \
diff --git a/Engine/Extras/RoboMerge/v3/src/common/dummyslackserver.ts b/Engine/Extras/RoboMerge/v3/src/common/dummyslackserver.ts
--- a/Engine/Extras/RoboMerge/v3/src/common/dummyslackserver.ts
+++ b/Engine/Extras/RoboMerge/v3/src/common/dummyslackserver.ts
@@ -1,9 +1,8 @@
 // Copyright Epic Games, Inc. All Rights Reserved.
 
-import { AppInterface, Handler, WebRequest } from './webserver'
 import { setDefault } from './helper'
 
-// This must be kept in sync with tests.ts
+// This must be kept in sync with framework.ts
 class DummyThread {
 	CL: number
 	target: string
@@ -19,115 +18,127 @@ const posted = new Map<string, Map<string, DummyThread[]>>()
 // channel -> timestamp -> messages
 const threads = new Map<string, Map<number, DummyThread>>()
 
-export class DummySlackApp implements AppInterface {
+class DummySlackAdmin {
+	static get conversations() {
+		return DummySlackConversations
+	}
+}
 
-	constructor(private req: WebRequest) {
+class DummySlackAppInternal {
+
+	static get client() {
+		return DummySlackClient
 	}
+}
 
-	@Handler('POST', '/api/*')
-	post(command: string) {
-		// @todo (maybe) construct app with logger
+class DummySlackChat {
 
-		const data = this.req.reqData
-		const getData = function() {
-			if (!data) {
-				throw new Error('No data to parse!')
-			}
-			let dataObj: any
-			try {
-				dataObj = JSON.parse(data)
-			}
-			catch (exc) {
-				console.log('Non-JSON sent to dummy-slack: ', data)
-				return { error: "Non-JSON sent to dummy-slack" }
-			}
-			return { dataObj }
+	static getPermalink(_: any) {
+		return { ok: true }		
+	}
+
+	static postMessage(args: any) {
+		if (args.thread_ts) {
+			let dummyThread = threads.get(args.channel)?.get(args.thread_ts)!
+			dummyThread.messages.push(args.attachments[0].text)
+			return { ok: true }
 		}
+		else {
+			let cl = args.cl
+			let target = args.target
+			if (cl && target) {
 
-		if (command === "chat.postMessage") {
-			if (!data) {
-				throw new Error('Nothing to post!')
-			}
-			const dataResult = getData()
-			if (dataResult.error) {
-				return dataResult
-			}
-			const dataObj = dataResult.dataObj
+				const dummyThread = new DummyThread
+				dummyThread.CL = cl
+				dummyThread.target = target
+				dummyThread.timestamp = Date.now()
+				dummyThread.messages = [args.attachments[0].text]
+
+				setDefault(
+					setDefault(posted, args.channel, new Map),
+					dummyThread.target, []
+				).push(dummyThread)
 
-			if (dataObj.thread_ts) {
-				let dummyThread = threads.get(dataObj.channel)?.get(dataObj.thread_ts)!
-				dummyThread.messages.push(dataObj.attachments[0].text)
-				return { ok: true }
+				setDefault(threads, args.channel, new Map)
+					.set(dummyThread.timestamp, dummyThread)
+
+				return { ok: true, ts: dummyThread.timestamp }
 			}
 			else {
-				let cl = dataObj.cl
-				let target = dataObj.target
-				if (cl && target) {
-
-					const dummyThread = new DummyThread
-					dummyThread.CL = cl
-					dummyThread.target = target
-					dummyThread.timestamp = Date.now()
-					dummyThread.messages = [dataObj.attachments[0].text]
-
-					setDefault(
-						setDefault(posted, dataObj.channel, new Map),
-						dummyThread.target, []
-					).push(dummyThread)
-
-					setDefault(threads, dataObj.channel, new Map)
-						.set(dummyThread.timestamp, dummyThread)
-
-					return { ok: true, ts: dummyThread.timestamp }
-				}
-				else {
-					console.log('dummy-slack POST: cl or edge not found', command, data)
-					return { error: "cl or edge not found"}
-				}
+				console.log('dummy-slack POST: cl or edge not found', JSON.stringify(args))
+				return { error: "cl or edge not found"}
 			}
-		}		
-		else if (command === "chat.update") {
-			// Not going to handle updating the stored messages on the dummy server
-			// until we have a reason to do so 
-			return { ok: true }
-		}
-		else if (command === "conversations.invite") {
-			return { ok: true }
-		}
-		else if (command === "conversations.open") {
-			const dataResult = getData()
-			if (dataResult.error) {
-				return dataResult
-			}
-			const dataObj = dataResult.dataObj
-			return { ok: true, channel: {id: `${dataObj.users}`} }
-		}
-		else if (command === "files.upload") {
-			// Not going to store files until we have a reason to do so
-			return { ok: true }
-		}
-		else {
-			console.error(`dummy-slack POST: unsupported command '${command}'`)
-			return { error: "unsupported command"}
 		}
 	}
 
-	@Handler('GET', '/posted/*')
-	channelMessages(channel: string) {
-		return [...(posted.get(channel) || [])]
+	static update(_: any) {
+		return { ok: true }		
 	}
+}
 
-	@Handler('GET', '/api/*')
-	get(command: string) {
-		if (command === "chat.getPermalink") {
-			// We don't need the permalink for any functional tests
-			return { ok: true }
-		}
-		else if (command.startsWith("users.lookupByEmail"))
-		{
-			return { ok: true, user: { id: `${this.req.url.searchParams.get('email')}` } }
-		}
-		return { ok: false }
+class DummySlackClient {
+
+	static get admin() {
+		return DummySlackAdmin
+	}
+
+	static get chat() {
+		return DummySlackChat
+	}
+
+	static get conversations() {
+		return DummySlackConversations
+	}
+
+	static get files() {
+		return DummySlackFiles
+	}
+
+	static get users() {
+		return DummySlackUsers
+	}
+
+}
+
+class DummySlackConversations {
+
+	static invite(_: any) {
+		return { ok: true }		
 	}
 
+	static open(args: any) {
+		return { ok: true, channel: {id: args.users} }
+	}
+
+}
+
+class DummySlackFiles {
+	static uploadV2(_: any) {
+		return { ok: true }		
+	}
+}
+
+class DummySlackUsers {
+	static lookupByEmail(args: any) {
+		return { ok: true, user: { id: args.email } }
+	}
+}
+
+export class DummySlackApp {
+
+	get app() {
+		return DummySlackAppInternal
+	}
+
+	get botToken() {
+		return ""
+	}
+
+	get userToken() {
+		return ""
+	}
+
+	static channelMessages(channel: string) {
+		return [...(posted.get(channel) || [])]
+	}
 }
\ No newline at end of file
diff --git a/Engine/Extras/RoboMerge/v3/src/robo/ipc.ts b/Engine/Extras/RoboMerge/v3/src/robo/ipc.ts
--- a/Engine/Extras/RoboMerge/v3/src/robo/ipc.ts
+++ b/Engine/Extras/RoboMerge/v3/src/robo/ipc.ts
@@ -11,6 +11,7 @@ import { roboAnalytics } from './roboanalytics';
 import { Status } from './status';
 import { getPreview } from './preview'
 import * as p4util from '../common/p4util';
+import { DummySlackApp } from '../common/dummyslackserver';
 
 // Reminder that RoboMerge variable is null in this file! ${RoboMerge}
 
@@ -31,6 +32,7 @@ const RobomergeMethodStrings = ['initialSubmit','merge_with_conflict','automerge
 type RobomergeMethods = typeof RobomergeMethodStrings[number];
 type MergeMethod = RobomergeMethods|'populate'|'manual_merge'
 
+// this must be kept in sync with framework.ts
 export type OperationReturnType = {
 	statusCode: number
 	message: string // Goal: to provide meaningful error messaging to the end user
@@ -69,6 +71,7 @@ export class IPC {
 			case 'crashAPI': throw new Error(`API Crash requested by ${msg.args![0]}`)
 			case 'forceBranchmapUpdate': return this.forceBranchmapUpdate(msg.args![0] as string)
 			case 'preview': return this.preview(msg.args![0] as string, msg.args![1] as (string|undefined))
+			case 'dummySlackChannelMessages': return this.dummySlackChannelMessages(msg.args![0] as string)
 
 			case 'doNodeOp': return this.doOperation(
 				msg.args![0] as string,
@@ -433,6 +436,10 @@ export class IPC {
 		}
 	}
 
+	private dummySlackChannelMessages(channel: string): OperationReturnType {
+		return { ...OPERATION_SUCCESS, data: DummySlackApp.channelMessages(channel) }
+	}
+
 	private dumpGraph(tagsObj: any): OperationReturnType {
 		let tags = new Set<string>()
 		for (const tag in tagsObj) {
diff --git a/Engine/Extras/RoboMerge/v3/src/robo/notifications.ts b/Engine/Extras/RoboMerge/v3/src/robo/notifications.ts
--- a/Engine/Extras/RoboMerge/v3/src/robo/notifications.ts
+++ b/Engine/Extras/RoboMerge/v3/src/robo/notifications.ts
@@ -13,8 +13,6 @@ import { BlockageNodeOpUrls } from './roboserver';
 import { Context } from './settings';
 import { Slack, SlackAttachment, SlackFile, SlackLinkButtonAction, SlackLinkButtonsAttachment, SlackMessageField, SlackMessage, SlackMessageStyles } from './slack';
 import { Vault } from './vault'
-import { WebServer } from '../common/webserver'
-import { DummySlackApp } from '../common/dummyslackserver'
 
 export const NOTIFICATIONS_PERSISTENCE_KEY = 'notifications'
 export const SLACK_MESSAGES_PERSISTENCE_KEY = 'slackMessages'
@@ -26,10 +24,8 @@ const EPIC_TIME_OPTIONS: DateTimeFormatOptions = {timeZone: 'EST5EDT', timeZoneN
 const KNOWN_BOT_NAMES = ['buildmachine', 'robomerge'];
 const KNOWN_BOT_EMAILS = [[REDACTED_EMAIL]'];
 
-let SLACK_TOKENS: {[name: string]: string} = {}
-const SLACK_DEV_DUMMY_TOKEN = 'dev'
-
 let args: Args
+let usingSlack: boolean = false
 export function notificationsInit(inArgs: Args, vault: Vault) {
 	
 	args = inArgs
@@ -38,16 +34,9 @@ export function notificationsInit(inArgs: Args, vault: Vault) {
 		Badge.setDevMode()
 	}
 
-	if (args.slackDomain.indexOf('localhost') >= 0) { // todo should parse out port
-		const slackServer = new WebServer(new ContextualLogger('dummy Slack'))
-		slackServer.addApp(DummySlackApp)
-		slackServer.open(8811, 'http')
-		return
-	}
-
-	const tokensObj = vault.slackTokens
-	if (tokensObj) {
-		SLACK_TOKENS = tokensObj
+	if (!args.devMode || args.useSlackInDev) {
+		Slack.init(vault.slackTokens)
+		usingSlack = true
 	}
 }
 
@@ -61,17 +50,19 @@ export async function postToRobomergeAlerts(message: string) {
 	if (!args.devMode) {
 		return postMessageToChannel(message, 'C044ZUSS71S')
 	}
+	return null
 }
 
 export async function postMessageToChannel(message: string, channel: string, style: SlackMessageStyles = SlackMessageStyles.GOOD) {
-	if (SLACK_TOKENS.bot) {
-		return new Slack({id: channel, botToken: SLACK_TOKENS.bot, userToken: SLACK_TOKENS.user}, args.slackDomain,new ContextualLogger(`Post Message to ${channel}`)).postMessageToDefaultChannel({
+	if (usingSlack) {
+		return new Slack(channel, new ContextualLogger(`Post Message to ${channel}`)).postMessageToDefaultChannel({
 			text: message,
 			style,
 			channel,
 			mrkdwn: true
 		})
 	}
+	return null
 }
 
 //////////
@@ -1045,11 +1036,9 @@ export function bindBotNotifications(events: BotEvents, slackChannelOverrides: [
 
 	let slackMessages
 
-	const botToken = (!args.devMode || args.useSlackInDev) && SLACK_TOKENS.bot || args.devMode && args.useSlackInDev && SLACK_DEV_DUMMY_TOKEN 
-	const userToken = (!args.devMode || args.useSlackInDev) && SLACK_TOKENS.user || args.devMode && args.useSlackInDev && SLACK_DEV_DUMMY_TOKEN
-	if (botToken && events.botConfig.slackChannel) {
+	if (usingSlack && events.botConfig.slackChannel) {
 		logger.info('Enabling Slack messages for ' +  events.botname)
-		slackMessages = new SlackMessages(new Slack({id: events.botConfig.slackChannel, botToken, userToken}, args.slackDomain, logger), persistence, logger)
+		slackMessages = new SlackMessages(new Slack(events.botConfig.slackChannel, logger), persistence, logger)
 	}
 		
 	events.registerHandler(new BotNotifications(events.botname, events.botConfig.slackChannel, externalUrl, blockageUrlGenerator, logger, slackMessages, slackChannelOverrides))
diff --git a/Engine/Extras/RoboMerge/v3/src/robo/robo.ts b/Engine/Extras/RoboMerge/v3/src/robo/robo.ts
--- a/Engine/Extras/RoboMerge/v3/src/robo/robo.ts
+++ b/Engine/Extras/RoboMerge/v3/src/robo/robo.ts
@@ -175,11 +175,6 @@ const COMMAND_LINE_ARGS: {[param: string]: Arg<any>} = {
 		env: 'SENTRY_DSN',
 		dflt: 'https:[REDACTED_EMAIL]/1432517' // Robomerge sentry project: https://sentry.io/organizations/to/issues/?project=1432517
 	},
-	slackDomain: {
-		match: /^-slackDomain=(.+)$/,
-		env: 'ROBO_SLACK_DOMAIN',
-		dflt: 'https://slack.com'
-	},
 	persistenceDir: {
 		match: /^-persistenceDir=(.+)$/,
 		env: 'ROBO_PERSISTENCE_DIR',
diff --git a/Engine/Extras/RoboMerge/v3/src/robo/roboserver.ts b/Engine/Extras/RoboMerge/v3/src/robo/roboserver.ts
--- a/Engine/Extras/RoboMerge/v3/src/robo/roboserver.ts
+++ b/Engine/Extras/RoboMerge/v3/src/robo/roboserver.ts
@@ -242,6 +242,13 @@ class RoboWebApp implements AppInterface {
 		return this.sendMessage('trackChange', [queryObj, tagsObj])
 	}
 
+	@Handler('GET', '/dummyslack/*')
+	channelMessages(channel: string) {
+		if (!ENVIRONMENT.devMode) {
+			return {statusCode: 403 /* Forbidden */, message: 'DevMode only feature'}
+		}
+		return this.sendMessage('dummySlackChannelMessages', [channel])
+	}
 
 	@Handler('GET', '/preview/*', {filetype: 'text/html'}) 
 	previewNoBot(clStr: string) {
diff --git a/Engine/Extras/RoboMerge/v3/src/robo/slack.ts b/Engine/Extras/RoboMerge/v3/src/robo/slack.ts
--- a/Engine/Extras/RoboMerge/v3/src/robo/slack.ts
+++ b/Engine/Extras/RoboMerge/v3/src/robo/slack.ts
@@ -1,13 +1,14 @@
 // Copyright Epic Games, Inc. All Rights Reserved.
 
-import * as request from '../common/request';
-import { ContextualLogger } from '../common/logger';
-import { URLSearchParams } from "url"
-
-export interface SlackChannel {
-	id: string
-	botToken: string
-	userToken: string
+const { App } = require('@slack/bolt')
+import { ContextualLogger } from "../common/logger"
+import { DummySlackApp } from "../common/dummyslackserver"
+
+interface SlackTokens {
+	app: string
+	bot: string
+	user: string
+	signature: string
 }
 
 export interface SlackMessageField {
@@ -79,141 +80,163 @@ export interface SlackFile {
 	title?: string
 }
 
-enum PostFormat {
-	JSON,
-	URL
-}
+const MAIN_MESSAGE_FIELDS = new Set(['username', 'icon_emoji', 'channel', 'target', 'cl'])
 
-const MAIN_MESSAGE_FIELDS = new Set(['username', 'icon_emoji', 'channel', 'target', 'cl']);
+class SlackApp {
+	readonly app: any
 
-export class Slack {
-	constructor(private channel: SlackChannel, private domain: string, private readonly logger: ContextualLogger) {
-	}
-
-	async addUserToChannel(user: string, channel: string, externalUser?: boolean) {
-		if (externalUser) {		
-			return this.post_user(this.channel.userToken || this.channel.botToken, 'admin.conversations.invite', {channel_id:channel, user_ids:user}, true)
-		}
-		return this.post('conversations.invite', {channel, users:user}, true)
+	constructor(readonly botToken: string, readonly userToken: string, appToken: string, signingSecret: string) {
+        this.app = new App({
+            token: botToken,
+            appToken,
+            signingSecret,
+            socketMode: true,
+          })
 	}
+}
 
-	async getChannelInfo(channel:string) {
-		return this.get('conversations.info', {channel})
-	}
+export class Slack {
 
-	async getPermalink(thread_ts: string, channel: string) {
-		return (await this.get('chat.getPermalink', {channel, message_ts: thread_ts})).permalink
-	}
+	// Assume a single slack connection for the time being, we can always rework this
+	// similar to the PerforceContext to have different apps in the future if needed
+	private static slackApp: SlackApp | DummySlackApp
 
-	async postMessage(message: SlackMessage) {
-		return (await this.post('chat.postMessage', this.makeArgs(message))).ts
+	static async init(slackTokens?: SlackTokens) {
+		if (slackTokens) {
+			Slack.slackApp = new SlackApp(slackTokens.bot, slackTokens.user, slackTokens.app, slackTokens.signature)
+			await Slack.slackApp.app.start()
+		}
+		else {
+			// if we don't have slack toekns then create the dummy slack app instead
+			Slack.slackApp = new DummySlackApp()
+		}
 	}
 
-	postMessageToDefaultChannel(message: SlackMessage) {
-		const args = this.makeArgs(message)
-		args.channel = this.channel.id
-		return this.post('chat.postMessage', args)
+	constructor(private channel: string, private logger: ContextualLogger) {
 	}
 
-	reply(thread_ts: string, message: SlackMessage) {
-		const args = this.makeArgs(message)
-		args.thread_ts = thread_ts
-		return this.post('chat.postMessage', args)
+	async addUserToChannel(users: string, channel: string, externalUser?: boolean) {
+		try {
+			const slackApp = Slack.slackApp
+			if (externalUser) {
+				return await slackApp.app.client.admin.conversations.invite({token: slackApp.userToken || slackApp.botToken, channel, users})
+			}
+			
+			return await slackApp.app.client.conversations.invite({token: slackApp.botToken, channel, users})
+		}
+		catch(error) {
+			// addUserToChannel can be expected to fail and the caller will handle error codes
+			return error.data
+		}
 	}
 
-	update(ts: string, message: SlackMessage) {
-		const args = this.makeArgs(message)
-		args.ts = ts
-		return this.post('chat.update', args)
+	async getChannelInfo(channel: string) {
+		try {
+			return (await Slack.slackApp.app.client.conversations.info({channel})).channel
+		}
+		catch(error) {
+			this.logger.error(error)
+			return {ok: false}
+		}
 	}
 
-	uploadFile(file: SlackFile) {
-		return this.post('files.upload', file, false, PostFormat.URL)
+	async getPermalink(thread_ts: string, channel: string) {
+		try {
+			return (await Slack.slackApp.app.client.chat.getPermalink({
+				channel,
+				message_ts: thread_ts
+			})).permalink
+		}
+		catch(error) {
+			this.logger.error(error)
+			return {ok: false}
+		}
 	}
 
-	listMessages(count?: number) {
-		const args: any = count ? {count} : count
-		// use channels.history if public?
-		return this.get('groups.history', args)
+	async postMessage(message: SlackMessage) {
+		try {
+			return (await Slack.slackApp.app.client.chat.postMessage(this.makeArgs(message))).ts
+		}
+		catch(error) {
+			this.logger.error(error)
+			return {ok: false}
+		}
 	}
 
-	async lookupUserIdByEmail(email: string) {
-		const userLookupResult = await this.get('users.lookupByEmail', {email}, true)
-		return userLookupResult.ok ? userLookupResult.user.id : null
+	postMessageToDefaultChannel(message: SlackMessage) {
+		message.channel = this.channel
+		return this.postMessage(message)
 	}
 
-	async openDMConversation(users: string | string[]) : Promise<string|null> {
-		if (users instanceof Array) {
-			users = users.join(',')
+	async reply(thread_ts: string, message: SlackMessage) {
+		try {
+			const args = this.makeArgs(message)
+			args.thread_ts = thread_ts
+			return await Slack.slackApp.app.client.chat.postMessage(args)
+		}
+		catch(error) {
+			this.logger.error(error)
+			return {ok: false}
 		}
-		const result = await this.post('conversations.open', {users})
-		return result.channel ? result.channel.id : null
 	}
 
-	/*private*/ async post_user(userToken: string, command: string, args: any, canFail? : boolean, format?: PostFormat) {
-		const resultJson = await request.post({
-			url: this.domain + '/api/' + command,
-			body: format == PostFormat.URL ? new URLSearchParams(Object.entries(args)).toString() : JSON.stringify(args),
-			headers: {Authorization: 'Bearer ' + userToken},
-			contentType: format == PostFormat.URL ? 'application/x-www-form-urlencoded' : 'application/json; charset=utf-8'
-		})
+	async update(ts: string, message: SlackMessage) {
 		try {
-			const result = JSON.parse(resultJson)
-			if (result.ok || canFail) {
-				return result
-			}
+			const args = this.makeArgs(message)
+			args.ts = ts
+			return await Slack.slackApp.app.client.chat.update(args)
 		}
-		catch {
+		catch(error) {
+			this.logger.error(error)
+			return {ok: false}
 		}
-
-		this.logger.error(`${command} generated:\n\t${resultJson}`)
-		return {ok: false}
 	}
 
-	/*private*/ async post(command: string, args: any, canFail? : boolean, format?: PostFormat) {
-		return this.post_user(this.channel.botToken, command, args, canFail, format)
+	async uploadFile(file: SlackFile) {
+		try {
+			const result =  await Slack.slackApp.app.client.files.uploadV2({
+				content: file.content,
+				channel_id: file.channels,
+				filename: file.filename,
+				filetype: file.filetype,
+				initial_comment: file.initial_comment,
+				thread_ts: file.thread_ts,
+				title: file.title
+			})
+
+			return (result.ok ? true : false)
+		}
+		catch(error) {
+			this.logger.error(error)
+			return false
+		}
 	}
 
-	/*private*/ async get(command: string, args: any, canFail? : boolean) {
-
-		// erg: why am I always passing a channel?
-		if (this.channel.id && !args.channel) {
-			args.channel = this.channel.id
+	async lookupUserIdByEmail(email: string) {
+		try {
+			return (await Slack.slackApp.app.client.users.lookupByEmail({email})).user.id
 		}
-
-		const qsBits: string[] = []
-		for (const arg in args) {
-			qsBits.push(`${encodeURIComponent(arg)}=${encodeURIComponent(args[arg])}`)
+		catch(error) {
+			this.logger.error(error)
+			return {ok: false}
 		}
+	}
 
-		const url = this.domain + `/api/${command}?${qsBits.join('&')}`
-
-		const rawResult = await request.get({url,
-			headers: {Authorization: 'Bearer ' + this.channel.botToken}
-		})
-
+	async openDMConversation(users: string | string[]) {
 		try {
-			const result = JSON.parse(rawResult)
-			if (result.ok || canFail) {
-				return result
+			if (users instanceof Array) {
+				users = users.join(',')
 			}
+
+			const result = await Slack.slackApp.app.client.conversations.open({users})
+			return (result.ok && result.channel ? result.channel.id : null)
 		}
-		catch {
+		catch(error) {
+			this.logger.error(error)
+			return null
 		}
-		
-		this.logger.error(`url: '${url}' error: '${rawResult}'`)
-		return {ok: false}
 	}
 
-// behaviour seems to be: put in attachment for green margin if any opts sent at all (e.g. {} different to nothing passed)
-
-
-// text in messageOpts was different to text argument:
-//	former goes in attachment
-
-//	should now overload so that string only puts text in args
-
-// look out for makeArgs with just string - support if necessary
 	private makeArgs(message: SlackMessage) {
 		const args: {[arg: string]: any} = {}
 
@@ -241,6 +264,10 @@ export class Slack {
 		// args.parse = 'client'
 		// args.parse = 'full'
 		args.link_names = 1
+
+		// this is required for inclusivity
+		args.text = message.text
+
 		return args
 	}
-}
+}
\ No newline at end of file
diff --git a/Engine/Extras/RoboMerge/v3/src/robo/vault.ts b/Engine/Extras/RoboMerge/v3/src/robo/vault.ts
--- a/Engine/Extras/RoboMerge/v3/src/robo/vault.ts
+++ b/Engine/Extras/RoboMerge/v3/src/robo/vault.ts
@@ -35,7 +35,7 @@ export class Vault {
     }
 
     get slackTokens() {
-        return this._vault['slack-tokens']
+        return this.valid ? this._vault['slack-tokens'] : null
     }
 
     get cookieKey() {
