Allow resolver to use the <!subteam syntax for group identification

diff --git a/Engine/Extras/RoboMerge/v3/src/common/perforce.ts b/Engine/Extras/RoboMerge/v3/src/common/perforce.ts
--- a/Engine/Extras/RoboMerge/v3/src/common/perforce.ts
+++ b/Engine/Extras/RoboMerge/v3/src/common/perforce.ts
@@ -1409,7 +1409,7 @@ export class PerforceContext {
 	// get the email (according to P4) for a specific user
 	async getEmail(username: string) {
 		let m = null
-		if (username.length > 0 && !username.startsWith('@')) {
+		if (username.length > 0 && !username.startsWith('@') && !username.startsWith('<')) {
 			const output = await this._execP4(null, ['user', '-o', username]);
 			// look for the email field
 			m = output.match(/\nEmail:\s+([^\n]+)\n/);
diff --git a/Engine/Extras/RoboMerge/v3/src/robo/edgebot.ts b/Engine/Extras/RoboMerge/v3/src/robo/edgebot.ts
--- a/Engine/Extras/RoboMerge/v3/src/robo/edgebot.ts
+++ b/Engine/Extras/RoboMerge/v3/src/robo/edgebot.ts
@@ -352,6 +352,12 @@ class EdgeBotImpl extends PerforceStatefulBot {
 			if (overriddenOwner.startsWith('@')) {
 				overriddenOwner = `@ ${overriddenOwner.substring(1)}`
 			}
+			else if (overriddenOwner.startsWith('<')) {
+				const subteamMatch = overriddenOwner.match(/<!subteam\^\w+\|([^>]+)>/)
+				if (subteamMatch) {
+					overriddenOwner = `@ ${subteamMatch[1]}`
+				}
+			}
 			description += `#ROBOMERGE-OWNER: ${overriddenOwner}\n`
 		}
 
diff --git a/Engine/Extras/RoboMerge/v3/src/robo/notifications.ts b/Engine/Extras/RoboMerge/v3/src/robo/notifications.ts
--- a/Engine/Extras/RoboMerge/v3/src/robo/notifications.ts
+++ b/Engine/Extras/RoboMerge/v3/src/robo/notifications.ts
@@ -71,7 +71,7 @@ export async function postMessageToChannel(message: string, channel: string, sty
 // make a Slack notifcation link for a user
 function atifyUser(user: string) {
 	// don't @ names of bots (currently making them tt style for Slack)
-	return isUserAKnownBot(user) ? `\`${user}\`` : user.startsWith('@') ? user : '@' + user
+	return isUserAKnownBot(user) ? `\`${user}\`` : user.startsWith('@') || user.startsWith('<') ? user : '@' + user
 }
 
 function generatePersistedSlackMessageKey(sourceCl: number, targetBranchArg: BranchArg, channel: string) {
@@ -513,7 +513,11 @@ export class BotNotifications implements BotEventHandler {
 
 		const userEmail = await blockage.ownerEmail
 		const slackUser = userEmail ? await this.slackMessages.getSlackUser(userEmail) : null
-		const channelPing = slackUser ? `<@${slackUser}>` : blockage.owner.startsWith('@') ? blockage.owner : `@${blockage.owner}`
+		const channelPing = slackUser 
+		                        ? `<@${slackUser}>` 
+		                        : blockage.owner.startsWith('@') || blockage.owner.startsWith('<')
+		                              ? blockage.owner 
+		                              : `@${blockage.owner}`
 
 		const isBotUser = isUserAKnownBot(blockage.owner)
 		const text =
