Remove orphaned debugging call

diff --git a/Engine/Extras/RoboMerge/v3/src/robo/autobranchupdater.ts b/Engine/Extras/RoboMerge/v3/src/robo/autobranchupdater.ts
--- a/Engine/Extras/RoboMerge/v3/src/robo/autobranchupdater.ts
+++ b/Engine/Extras/RoboMerge/v3/src/robo/autobranchupdater.ts
@@ -242,7 +242,6 @@ export class AutoBranchUpdater implements Bot {
 		if (fstat.length == 0) {
 			await this.p4.add(workspace, cl, depotpath)
 		}
-		JSON.stringify(await this.p4.describe(cl))
 		await this.p4.submit(workspace, cl)
 	}
 
