Multi-server support

diff --git a/Engine/Extras/RoboMerge/v3/src/common/p4util.ts b/Engine/Extras/RoboMerge/v3/src/common/p4util.ts
--- a/Engine/Extras/RoboMerge/v3/src/common/p4util.ts
+++ b/Engine/Extras/RoboMerge/v3/src/common/p4util.ts
@@ -9,17 +9,23 @@ const USER_WORKSPACE_EXCLUDE_PATTERNS: (RegExp | string)[] = [
 	'swarm-'
 ]
 
+type WorkspaceID = {
+	serverAddress?: string
+	workspaceName: string
+}
 // true has been cleaned, false is being cleaned, undefined needs cleaning
-let cleanedWorkspaces: Map<string, boolean> = new Map()
+let cleanedWorkspaces: Map<WorkspaceID, boolean> = new Map()
 
 export async function cleanWorkspace(logger: ContextualLogger, p4: PerforceContext, workspace: Workspace, workspaceRoot: string, edgeServerAddress?: string) {
 	
+	const workspaceID: WorkspaceID = {serverAddress: p4.serverAddress, workspaceName: workspace.name}
+
 	// if the workspace is being cleaned, wait for it to finish
-	let workspaceCleaned = cleanedWorkspaces.get(workspace.name)
+	let workspaceCleaned = cleanedWorkspaces.get(workspaceID)
 	while (workspaceCleaned === false) {
 		logger.info(`Waiting on ${workspace.name} being cleaned`)
 		await _setTimeout(1000)
-		workspaceCleaned = cleanedWorkspaces.get(workspace.name)
+		workspaceCleaned = cleanedWorkspaces.get(workspaceID)
 	}
 
 	if (workspaceCleaned) {
@@ -27,7 +33,7 @@ export async function cleanWorkspace(logger: ContextualLogger, p4: PerforceConte
 	}
 
 	logger.info(`Cleaning ${workspace.name}`)
-	cleanedWorkspaces.set(workspace.name, false)
+	cleanedWorkspaces.set(workspaceID, false)
 	const changes = await p4.get_pending_changes(workspace.name) as Change[]
 
 	// revert any pending changes left over from previous runs
@@ -58,7 +64,7 @@ export async function cleanWorkspace(logger: ContextualLogger, p4: PerforceConte
 	logger.info(`Resetting workspace ${workspace.name} to revision 0`)
 	await p4.sync(workspace, workspaceRoot + '#0', {edgeServerAddress})
 
-	cleanedWorkspaces.set(workspace.name, true)
+	cleanedWorkspaces.set(workspaceID, true)
 }
 
 export async function getWorkspacesForUser(p4: PerforceContext, user: string) {
diff --git a/Engine/Extras/RoboMerge/v3/src/common/perforce.ts b/Engine/Extras/RoboMerge/v3/src/common/perforce.ts
--- a/Engine/Extras/RoboMerge/v3/src/common/perforce.ts
+++ b/Engine/Extras/RoboMerge/v3/src/common/perforce.ts
@@ -5,6 +5,7 @@ import * as path from 'path';
 import { postToRobomergeAlerts } from '../robo/notifications';
 import { roboAnalytics } from '../robo/roboanalytics';
 import { ContextualLogger, NpmLogLevel } from './logger';
+import { Vault } from '../robo/vault'
 import { VersionReader } from './version';
 
 type Type = 'string' | 'integer' | 'boolean'
@@ -317,6 +318,7 @@ export interface ConflictedResolveNFile {
 
 interface ServerConfig {
 	username: string
+	serverAddress?: string
 	serverVersion: number
 	multiServerEnvironment: boolean
 }
@@ -453,7 +455,7 @@ const MIN_SERVER_VERSION: number = 2020.1
  * This method must succeed before PerforceContext can be used. Otherwise retrieving the Perforce context through
  * getServerContext() will error.
  */
-export async function initializePerforce(logger: ContextualLogger) {
+export async function initializePerforce(logger: ContextualLogger, vault: Vault) {
 	setInterval(() => {
 		const now = Date.now()
 		for (const command of runningPerforceCommands) {
@@ -464,8 +466,16 @@ export async function initializePerforce(logger: ContextualLogger) {
 		}
 	}, 5*60*1000)
 
-	logger.info("Robomerge running in single server mode")
-	await PerforceContext.initServer(logger, "default")
+	const servers = vault.valid && vault.perforceServers
+
+	if (servers) {
+		await Promise.all(Array.from(servers).map(([server, serverDetails]) => PerforceContext.initServer(logger, server, serverDetails)))
+	}
+	else {
+		logger.info("Robomerge running in single server mode")
+		await PerforceContext.initServer(logger, "default")
+	}
+
 }
 
 /**
@@ -484,7 +494,12 @@ export class PerforceContext {
 		
 		const serverAddress = serverDetails && serverDetails.port
 
-		const output = await PerforceContext.execAndParse(logger, null, ["login", "-s"], { serverAddress });
+		let loginArgs = ["login", "-s"]
+		if (serverDetails && serverDetails.user) {
+			loginArgs = ["-u", serverDetails.user, ...loginArgs]
+		}
+
+		const output = await PerforceContext.execAndParse(logger, null, loginArgs, { serverAddress });
 		let resp = output[0];
 	
 		if (resp && resp.User) {
@@ -502,6 +517,7 @@ export class PerforceContext {
 
 			let serverConfig: ServerConfig = {
 				username: resp.User,
+				serverAddress,
 				serverVersion,
 				multiServerEnvironment: serversOutput.length > 1
 			}
@@ -528,10 +544,18 @@ export class PerforceContext {
 		return new PerforceContext(logger, server)
 	}
 
+	static getServerConfig(serverID?: string) {
+		return this.servers.get(serverID || 'default')
+	}
+
 	get username() {
 		return this._serverConfig.username
 	}
 
+	get serverAddress() {
+		return this._serverConfig.serverAddress
+	}
+
 	get serverVersion() {
 		return this._serverConfig.serverVersion
 	}
@@ -666,7 +690,7 @@ export class PerforceContext {
 
 	async stream(streamName: string) {
 		try {
-			const stream = (await this.execAndParse(null, ['streams', streamName], {quiet: false}))[0]
+			const stream = (await this.execAndParse(null, ['streams', streamName]))[0]
 			const streamSpec: StreamSpec = {
 				stream: stream.Stream as string,
 				name: stream.Name as string,
@@ -796,14 +820,9 @@ export class PerforceContext {
 
 	// sync the depot path specified
 	async sync(roboWorkspace: RoboWorkspace, depotPath: string, params?: SyncParams) {
-		let args = params && params.edgeServerAddress ? ['-p', params.edgeServerAddress] : []
-		args.push('sync')
-		if (params && params.opts) {
-			args = [...args, ...params.opts]
-		}
-		args.push(depotPath)
+		const args = ['sync', ...(params && params.opts || []), depotPath]
 		try {
-			return await this.execAndParse(roboWorkspace, args)
+			return await this.execAndParse(roboWorkspace, args, {serverAddress: params && params.edgeServerAddress})
 		}
 		catch (reason) {
 			if (!isExecP4Error(reason)) {
@@ -821,7 +840,7 @@ export class PerforceContext {
 		return null
 	}
 
-	async newWorkspace(workspaceName: string, params: any, edgeServer?: EdgeServer) {
+	newWorkspace(workspaceName: string, params: any, edgeServer?: EdgeServer) {
 		params.Client = workspaceName
 		if (!('Root' in params)) {
 			params.Root = 'd:/ROBO/' + workspaceName // default windows path
@@ -943,6 +962,82 @@ export class PerforceContext {
 		}
 	}
 
+	async transfer(targetP4: PerforceContext, roboWorkspace: RoboWorkspace, source: IntegrationSource, dest_changelist: number, target: IntegrationTarget)
+	: Promise<string[]> {
+
+		const desc = await this.describe(source.changelist)
+
+		let failures: string[] = []
+
+		const sourceStreamLength = source.stream!.length
+
+		await Promise.all(desc.entries.map(entry => {
+			const targetDepotFile = target.stream! + entry.depotFile.substring(sourceStreamLength)
+			switch(entry.action) {
+				case 'add':
+				case 'branch': {
+					// Validate file doesn't exist
+					return targetP4.fstat(roboWorkspace, targetDepotFile)
+					.then(fstat => {
+						if (fstat.length == 0 || fstat[0].headAction.includes('delete')) {
+							return targetP4.where(roboWorkspace, targetDepotFile)
+							.then(where => { return this.print(`${entry.depotFile}@${source.changelist}`, where[0].path) })
+							.then(() => { return targetP4.add(roboWorkspace, dest_changelist, targetDepotFile, entry.type) })
+						}
+						else {
+							failures.push(`File ${targetDepotFile} already exists when being added in CL# ${source.changelist}`)
+							return
+						}
+					})
+				}
+				case 'edit':
+				case 'integrate': {
+					// Validate file exists
+					return targetP4.fstat(roboWorkspace, targetDepotFile)
+					.then(fstat => {
+						if (fstat.length > 0 && !fstat[0].headAction.includes('delete')) {
+							return targetP4.sync(roboWorkspace, targetDepotFile, {opts: ['-k']})
+							.then(() => { return targetP4.edit(roboWorkspace, dest_changelist, targetDepotFile, ['-k', '-t', entry.type]) })
+							.then(() => { return this.print(`${entry.depotFile}@${source.changelist}`, fstat[0].clientFile) })	
+						}
+						else {
+							failures.push(`File ${targetDepotFile} does not exist when being edited in CL# ${source.changelist}`)
+							return
+						}
+					})
+				}
+				case 'delete': {
+					return targetP4.delete(roboWorkspace, dest_changelist, targetDepotFile)
+				}
+				case 'move/delete': {
+					return this.integrated(null, entry.depotFile, {intoOnly: true, startCL: source.changelist})
+					.then(integrated => {
+						const integ = integrated.find(integ => parseInt(integ.change) == source.changelist)
+						if (integ) {
+							const targetMoveAddDepotFile = target.stream! + integ.fromFile.substring(sourceStreamLength)
+							return targetP4.sync(roboWorkspace, targetDepotFile)
+							.then(() => { return targetP4.edit(roboWorkspace, dest_changelist, targetDepotFile) })
+							.then(() => { return targetP4.move(roboWorkspace, dest_changelist, targetDepotFile, targetMoveAddDepotFile, ['-f']) })
+							.then(move => {	return this.print(`${integ.fromFile}@${source.changelist}`, move[0].path) })
+						}
+						else {
+							failures.push(`Unable to look up integrated record for ${entry.depotFile} in CL# ${source.changelist}`)
+							return
+						}
+					})
+				}
+				case 'move/add': return // move/add is handled from the move/delete portion of the pair
+
+				default: {
+					failures.push(`Unhandled transfer action ${entry.action} for ${entry.depotFile} in CL# ${source.changelist}`)
+					return
+				}
+			}
+		}))
+
+		return failures
+	}
+
 	// integrate a CL from source to destination, resolve, and place the results in a new CL
 	// output format is true if the integration resolved or false if the integration wasn't necessary (still considered a success)
 	// failure to resolve is treated as an error condition
@@ -1332,6 +1427,13 @@ export class PerforceContext {
 		}
 	}
 
+	revertAndDelete(roboWorkspace: string, cl: number, edgeServerAddress?: string) {
+		return this.revert(roboWorkspace, cl, [], edgeServerAddress)
+		.then(() => {
+			return this.deleteCl(roboWorkspace, cl, edgeServerAddress);	
+		})
+	}
+
 	// list files in changelist that need to be resolved
 	async listFilesToResolve(roboWorkspace: RoboWorkspace, changelist: number) {
 		try {
@@ -1346,7 +1448,7 @@ export class PerforceContext {
 	}
 
 	async move(roboWorkspace: RoboWorkspace, cl: number, src: string, target: string, opts?: string[]) {
-		return this._execP4(roboWorkspace, ['move', ...(opts || []), '-c', cl.toString(), src, target]);
+		return this.execAndParse(roboWorkspace, ['move', ...(opts || []), '-c', cl.toString(), src, target]);
 	}
 
 	/**
@@ -1439,6 +1541,17 @@ export class PerforceContext {
 		return this._execP4(roboWorkspace, args);
 	}
 
+	/** Mark a file for delete into a specific changelist ** ASYNC ** */
+	delete(roboWorkspace: RoboWorkspace, cl: number, filePath: string) {
+		const args = [
+			'delete', 
+			'-c', cl.toString(),
+			'-v',
+			filePath
+		]
+		return this._execP4(roboWorkspace, args);
+	}
+
 	async describe(cl: number, maxFiles?: number, includeShelved?: boolean) {
 		const args = ['describe',
 			...(maxFiles ? ['-m', maxFiles.toString()] : []),
@@ -1596,8 +1709,13 @@ export class PerforceContext {
 		}
 	}
 
-	print(path: string) {
-		return this._execP4(null, ['print', '-q', path], { noCwd: true, quiet: true })
+	print(path: string, outputPath?: string) {
+		let args = ['print', '-q']
+		if (outputPath) {
+			args.push('-o', outputPath)
+		}
+		args.push(path)
+		return this._execP4(null, args, { noCwd: true, quiet: false })
 	}
 
 	private _sanitizeDescription(description: string) {
@@ -1786,6 +1904,8 @@ export class PerforceContext {
 	}
 
 	private _execP4(roboWorkspace: RoboWorkspace, args: string[], optsIn?: ExecOpts) {
+		optsIn = optsIn || {}
+		optsIn.serverAddress = optsIn.serverAddress || this._serverConfig.serverAddress
 		return PerforceContext._execP4(this.logger, roboWorkspace, ['-u', this.username, ...args], optsIn)
 	}
 
@@ -1799,6 +1919,8 @@ export class PerforceContext {
 	}
 
 	private async _execP4Ztag(roboWorkspace: RoboWorkspace, args: string[], opts?: ExecZtagOpts) {
+		opts = opts || {}
+		opts.serverAddress = opts.serverAddress || this._serverConfig.serverAddress
 		return PerforceContext._execP4Ztag(this.logger, roboWorkspace, ['-u', this.username, ...args], opts)
 	}
 
@@ -1872,6 +1994,8 @@ export class PerforceContext {
 	}
 
 	async execAndParse(roboWorkspace: RoboWorkspace, args: string[], execOptions?: ExecOpts, parseOptions?: ParseOptions) {
+		execOptions = execOptions || {}
+		execOptions.serverAddress = execOptions.serverAddress || this._serverConfig.serverAddress
 		return PerforceContext.execAndParse(this.logger, roboWorkspace, ['-u', this.username, ...args], execOptions, parseOptions)
 	}
 
@@ -1952,7 +2076,9 @@ export class PerforceContext {
 	}
 
 	async execAndParseArray(roboWorkspace: RoboWorkspace, args: string[], execOptions?: ExecOpts, headerOptions?: ParseOptions, arrayEntryOptions?: ParseOptions) {
-		return PerforceContext.execAndParseArray(this.logger, roboWorkspace, args, execOptions, headerOptions, arrayEntryOptions)
+		execOptions = execOptions || {}
+		execOptions.serverAddress = execOptions.serverAddress || this._serverConfig.serverAddress
+		return PerforceContext.execAndParseArray(this.logger, roboWorkspace, ['-u', this.username, ...args], execOptions, headerOptions, arrayEntryOptions)
 	}
 }
 
diff --git a/Engine/Extras/RoboMerge/v3/src/robo/autobranchupdater.ts b/Engine/Extras/RoboMerge/v3/src/robo/autobranchupdater.ts
--- a/Engine/Extras/RoboMerge/v3/src/robo/autobranchupdater.ts
+++ b/Engine/Extras/RoboMerge/v3/src/robo/autobranchupdater.ts
@@ -149,7 +149,7 @@ export class AutoBranchUpdater implements Bot {
 		}
 
 		const validationErrors: string[] = []
-		const result = await BranchDefs.parseAndValidate(this.p4, validationErrors, branchGraphText)
+		const result = await BranchDefs.parseAndValidate(this.abuLogger, validationErrors, branchGraphText)
 		if (!result.branchGraphDef) {
 			// @todo email author of changes!
 			let errText = 'failed to parse/validate branch specs file\n'
diff --git a/Engine/Extras/RoboMerge/v3/src/robo/branchdefs.ts b/Engine/Extras/RoboMerge/v3/src/robo/branchdefs.ts
--- a/Engine/Extras/RoboMerge/v3/src/robo/branchdefs.ts
+++ b/Engine/Extras/RoboMerge/v3/src/robo/branchdefs.ts
@@ -1,6 +1,7 @@
 // Copyright Epic Games, Inc. All Rights Reserved.
 
 import { PerforceContext } from '../common/perforce'
+import { ContextualLogger } from '../common/logger';
 
 const jsonlint: any = require('jsonlint')
 
@@ -114,6 +115,7 @@ const nodeOptionFieldsPrototype = {
 	blockAssetFlow: [''],
 	disallowDeadend: false,
 
+	streamServer: undefined as string | undefined,
 	streamDepot: '',
 	streamName: '',
 	streamSubpath: '',
@@ -263,7 +265,7 @@ export class BranchDefs {
 		}
 	}
 
-	static async parseAndValidate(p4: PerforceContext, outErrors: string[], branchSpecsText: string, isPreviewing?: boolean): Promise<ParseResult> {
+	static async parseAndValidate(logger: ContextualLogger, outErrors: string[], branchSpecsText: string, isPreviewing?: boolean): Promise<ParseResult> {
 		const defaultConfigForWholeBot: BotConfig = {
 			defaultStreamDepot: null,
 			defaultIntegrationMethod: null,
@@ -378,7 +380,7 @@ export class BranchDefs {
 				def.streamSubpath
 			)
 
-			if (!isPreviewing && streamResult.stream && !await p4.stream(streamResult.stream)) {
+			if (!isPreviewing && streamResult.stream && !await PerforceContext.getServerContext(logger, def.streamServer).stream(streamResult.stream)) {
 				outErrors.push(`Stream ${streamResult.stream} not found`)
 			}
 		}
diff --git a/Engine/Extras/RoboMerge/v3/src/robo/edgebot.ts b/Engine/Extras/RoboMerge/v3/src/robo/edgebot.ts
--- a/Engine/Extras/RoboMerge/v3/src/robo/edgebot.ts
+++ b/Engine/Extras/RoboMerge/v3/src/robo/edgebot.ts
@@ -115,7 +115,7 @@ class EdgeBotImpl extends PerforceStatefulBot {
 		this.graphBotName = sourceNode.graphBotName
 		this.edgeBotLogger = new ContextualLogger(this.fullNameForLogging)
 
-		this.p4 = PerforceContext.getServerContext(this.edgeBotLogger)
+		this.p4 = PerforceContext.getServerContext(this.edgeBotLogger, sourceNode.branch.config.streamServer)
 		this.eventTriggers.onChangeParsed((info: ChangeInfo) => this.onGlobalChange(info))
 		this.gate = new Gate(
 			{ from: this.sourceBranch
@@ -298,7 +298,7 @@ class EdgeBotImpl extends PerforceStatefulBot {
 		}
 
 		this.edgeBotLogger.info(`${logMessage}. Reverting ${pending.newCl}.`)
-		await this.revertAndDelete(coercePerforceWorkspace(pending.change.targetWorkspace)!.name, pending.newCl)
+		await this.p4.revertAndDelete(coercePerforceWorkspace(pending.change.targetWorkspace)!.name, pending.newCl)
 
 		let pauseDurationSeconds = FAILED_CHANGELIST_PAUSE_TIMEOUT_SECONDS
 		// if we have a target, make sure pause duration is at least 2x duration of failed integration
@@ -443,8 +443,13 @@ class EdgeBotImpl extends PerforceStatefulBot {
 
 		target.description = description + additionalDescription
 
-		// do the integration
-		return this.integrate(info, target)
+		if (target.branch.config.streamServer != this.sourceBranch.config.streamServer) {
+			return this.transfer(info, target)
+		}
+		else {
+			// do the integration
+			return this.integrate(info, target)
+		}
 	}
 
 	public resetIntegrationTimestamp() {
@@ -460,8 +465,9 @@ class EdgeBotImpl extends PerforceStatefulBot {
 		this.sourceNode.queueEdgeUnblock(this.targetBranch.upperName)
 	}
 
-	public async getWorkspace(edgeServer?: EdgeServer) {
+	public async getWorkspace(edgeServer?: EdgeServer, p4? : PerforceContext) {
 
+		p4 = p4 || this.p4
 		const edgeServerAddress: string | undefined = edgeServer && edgeServer.address
 
 		// name the workspace
@@ -476,7 +482,7 @@ class EdgeBotImpl extends PerforceStatefulBot {
 				workspaceName += fromName
 			}
 		}
-		const p4username = this.p4.username
+		const p4username = p4.username
 		if (p4username !== 'robomerge') {
 			workspaceName = [p4username!.toUpperCase(), process.platform.toUpperCase(), workspaceName].join('_')
 		}
@@ -493,12 +499,12 @@ class EdgeBotImpl extends PerforceStatefulBot {
 		}
 
 		// do we already have a workspace?
-		const existingWorkspaceInfo = await this.p4.find_workspace_by_name(workspaceName, {edgeServerAddress, includeUnloaded: true})
+		const existingWorkspaceInfo = await p4.find_workspace_by_name(workspaceName, {edgeServerAddress, includeUnloaded: true})
 		if (existingWorkspaceInfo.length > 0) {
 			if (existingWorkspaceInfo[0].IsUnloaded) {
-				await this.p4.reloadWorkspace(workspaceName, edgeServerAddress)
+				await p4.reloadWorkspace(workspaceName, edgeServerAddress)
 			}
-			await p4util.cleanWorkspace(this.logger, this.p4, {name: workspaceName, directory: path}, this.targetBranch.rootPath, edgeServerAddress)
+			await p4util.cleanWorkspace(this.logger, p4, {name: workspaceName, directory: path}, this.targetBranch.rootPath, edgeServerAddress)
 		}
 		else {
 			const params: any = {};
@@ -511,7 +517,7 @@ class EdgeBotImpl extends PerforceStatefulBot {
 				];
 			}
 
-			await this.p4.newGraphBotWorkspace(workspaceName, params, edgeServer);
+			await p4.newGraphBotWorkspace(workspaceName, params, edgeServer);
 
 			// if we're on linux, remove the directory whenever we create the workspace for the first time
 			if (process.platform === "linux") {
@@ -522,7 +528,7 @@ class EdgeBotImpl extends PerforceStatefulBot {
 				require('child_process').execSync(`rm -rf ${dir}/*`);
 			}
 			else {
-				await this.p4.clean(workspaceName);
+				await p4.clean(workspaceName);
 			}
 		}
 		
@@ -607,7 +613,7 @@ class EdgeBotImpl extends PerforceStatefulBot {
 
 		// Revert attempt
 		if (pending.newCl > 0) {
-			await this.revertAndDelete(info.targetWorkspace, pending.newCl, edgeServerAddress)
+			await this.p4.revertAndDelete(info.targetWorkspace, pending.newCl, edgeServerAddress)
 			pending.newCl = -1
 		}
 		
@@ -621,22 +627,91 @@ class EdgeBotImpl extends PerforceStatefulBot {
 		if (await this.sourceNode.handleMergeFailure(failure, pending)) {
 			this.gate.onBlockage()
 		}
-		return new EdgeIntegrationDetails('error', errors.join('\n'))
+		return new EdgeIntegrationDetails('error', description)
+	}
+
+	private async transfer(info: ChangeInfo, target: MergeAction) : Promise<EdgeIntegrationDetails> {
+		const to_integrate = info.cl
+		this._log_action(`Integrating CL ${to_integrate} to ${this.targetBranch.name}`)
+
+		// if required, add author review here so they're not in target.description, which is used for shelf description in case of conflict
+		const desc = target.description! // target.description always ends in newline
+
+		const targetServerP4 = PerforceContext.getServerContext(this.logger, target.branch.config.streamServer)
+		info.targetWorkspace = await this.getWorkspace(undefined, targetServerP4)
+
+		// create a new CL
+		const changenum = await targetServerP4.new_cl(info.targetWorkspace, desc)
+
+		// try to integrate
+		const branchSpecToTarget = this.sourceBranch.branchspec.get(this.targetBranch.upperName)
+		const source : IntegrationSource = {
+			branchspec: branchSpecToTarget,
+			changelist: to_integrate,
+			depot: this.sourceBranch.depot,
+			path_from: this.sourceBranch.rootPath,
+			stream: this.sourceBranch.stream
+		}
+		const integTarget : IntegrationTarget = {
+			depot: this.targetBranch.depot,
+			path_to: this.targetBranch.rootPath,
+			stream: this.targetBranch.stream
+		}
+
+		this.currentIntegrationStartTimestamp = Date.now()
+		await targetServerP4.sync(info.targetWorkspace, this.targetBranch.rootPath + '#0')
+		const failures = await this.p4.transfer(targetServerP4, info.targetWorkspace, source, changenum, integTarget)
+
+		const pending: PendingChange = {change: info, action: target, newCl: changenum}
+
+		let description = failures.join('\n')
+		let failure: Failure | null = null
+
+		if (failures.length == 0) {
+			// no failures, try to submit
+			const result = await this.submitChangelist(pending, 0, targetServerP4)
+			if (result.result !== 'error') {
+				return result
+			}
+
+			let details: string | undefined
+			const match = result.message.match(/.*STDERR:([^]*)STDOUT:/)
+			if (match)
+			{
+				details = match[1].trim()
+			}
+
+			failure = { kind: 'Commit failure', description: result.message, details }
+			description = `${failure.kind}: ${failure.description}`
+		}
+		else {
+			failure  = { kind: 'Transfer error', description }
+		}
+
+		// Revert attempt
+		if (pending.newCl > 0) {
+			await targetServerP4.revertAndDelete(info.targetWorkspace, pending.newCl)
+			pending.newCl = -1
 	}
 
-	private async revertAndDelete(roboWorkspace: string, cl: number, edgeServerAddress?: string) {
-		this._log_action(`Reverting CL ${cl}`);
-		await this.p4.revert(roboWorkspace, cl, [], edgeServerAddress);
+		if (!pending.change.userRequest) {
+			// Pause Edge
+			const pauseInfo = this.createEdgeBlockageInfo(failure, pending)
+			this.block(pauseInfo, FAILED_CHANGELIST_PAUSE_TIMEOUT_SECONDS)
+		}
 
-		this._log_action(`Deleting CL ${cl}`);
-		await this.p4.deleteCl(roboWorkspace, cl, edgeServerAddress);
+		// Send to source node to facilitate notification handling
+		if (await this.sourceNode.handleMergeFailure(failure, pending)) {
+			this.gate.onBlockage()
+		}
+		return new EdgeIntegrationDetails('error', description)
 	}
 
 	async revertPendingCLWithShelf(client: string, change: number, userContext: string) {
 		this._log_action(` ${userContext} - Deleting shelf from and reverting ${change}`)
 		try {
 			this.edgeBotLogger.info(await this.p4.delete_shelved(client, change))
-			await this.revertAndDelete(client, change)
+			await this.p4.revertAndDelete(client, change)
 		}
 		catch (err) {
 			this.edgeBotLogger.printException(err, `${userContext} - Error reverting ${change}`)
@@ -743,15 +818,17 @@ class EdgeBotImpl extends PerforceStatefulBot {
 	 * @param {Number} [resolveRetries=3] Number of times to catch non-fatal errors and attempt the resolve -> submit chain again.
 	 * @todo In a perfect world, parameter 'resolveRetries' wouldn't exist and an orchastration method would handle this in a loop.
 	 */
-	async submitChangelist(pending: PendingChange, resolveRetries = 3): Promise<EdgeIntegrationDetails> {
+	async submitChangelist(pending: PendingChange, resolveRetries = 3, targetServerP4?: PerforceContext): Promise<EdgeIntegrationDetails> {
 		const info = pending.change
 		const target = pending.action
 		const changenum = pending.newCl
 		const edgeServerAddress = pending.change.edgeServerToHostShelf && pending.change.edgeServerToHostShelf.address
 
+		const p4 = targetServerP4 || this.p4
+
 		// try to submit
 		this._log_action(`Submitting CL ${changenum} by ${info.author}`)
-		const result = await this.p4.submit(
+		const result = await p4.submit(
 			pending.change.targetWorkspace, 
 			changenum,
 			edgeServerAddress)
@@ -787,7 +864,7 @@ class EdgeBotImpl extends PerforceStatefulBot {
 			// change owner, so users can edit change descriptions later for reconsideration
 			this.edgeBotLogger.info(`Setting owner of CL ${finalCl} to author of change: ${info.author}`);
 			try {
-				await this.p4.editOwner(info.targetWorkspace, finalCl, info.author, {changeSubmitted: true, edgeServerAddress})
+				await p4.editOwner(info.targetWorkspace, finalCl, info.author, {changeSubmitted: true, edgeServerAddress})
 			}
 			catch (reason) {
 				let errPreface = 'Error changing owner'
diff --git a/Engine/Extras/RoboMerge/v3/src/robo/graphbot.ts b/Engine/Extras/RoboMerge/v3/src/robo/graphbot.ts
--- a/Engine/Extras/RoboMerge/v3/src/robo/graphbot.ts
+++ b/Engine/Extras/RoboMerge/v3/src/robo/graphbot.ts
@@ -64,7 +64,7 @@ export class GraphBot implements BotEventHandler {
 		const fileText = require('fs').readFileSync(branchSettingsPath, 'utf8')
 
 		const validationErrors: string[] = []
-		const result = await BranchDefs.parseAndValidate(graphBot.p4, validationErrors, fileText)
+		const result = await BranchDefs.parseAndValidate(graphBot.botLogger,validationErrors, fileText)
 		if (!result.branchGraphDef) {
 			throw new Error(validationErrors.length === 0 ? 'Failed to parse' : validationErrors.join('\n'))
 		}
diff --git a/Engine/Extras/RoboMerge/v3/src/robo/nodebot.ts b/Engine/Extras/RoboMerge/v3/src/robo/nodebot.ts
--- a/Engine/Extras/RoboMerge/v3/src/robo/nodebot.ts
+++ b/Engine/Extras/RoboMerge/v3/src/robo/nodebot.ts
@@ -106,7 +106,7 @@ export class NodeBot extends PerforceStatefulBot implements NodeBotInterface {
 
 		this.externalRobomergeUrl = externalUrl
 
-		this.p4 = PerforceContext.getServerContext(this.nodeBotLogger)
+		this.p4 = PerforceContext.getServerContext(this.nodeBotLogger, branchDef.config.streamServer)
 
 		this.mailer = mailer
 
diff --git a/Engine/Extras/RoboMerge/v3/src/robo/preview.ts b/Engine/Extras/RoboMerge/v3/src/robo/preview.ts
--- a/Engine/Extras/RoboMerge/v3/src/robo/preview.ts
+++ b/Engine/Extras/RoboMerge/v3/src/robo/preview.ts
@@ -5,7 +5,6 @@ import { PerforceContext } from '../common/perforce';
 import { Status } from './status';
 
 const logger = new ContextualLogger('preview')
-let p4: PerforceContext | null = null
 
 // let branchSpecsRootPath = ''
 export function init(_root: string) {
@@ -13,10 +12,8 @@ export function init(_root: string) {
 }
 
 export async function getPreview(cl: number, singleBot?: string) {
-	// @todo print content from shelf
-	if (!p4) {
-		p4 = PerforceContext.getServerContext(logger)
-	}
+
+	let p4 = PerforceContext.getServerContext(logger)
 
 	const status = new Status(new Date(), 'preview', logger)
 
@@ -38,7 +35,7 @@ export async function getPreview(cl: number, singleBot?: string) {
 		const fileText = await p4.print(`${path}@=${cl}`)
 
 		let validationErrors: string[] = []
-		const result = await BranchDefs.parseAndValidate(p4, validationErrors, fileText, true)
+		const result = await BranchDefs.parseAndValidate(logger, validationErrors, fileText, true)
 
 		const errorPrefix = `\n\t${bot} validation failed: `
 		if (!result.branchGraphDef) {
diff --git a/Engine/Extras/RoboMerge/v3/src/robo/robo.ts b/Engine/Extras/RoboMerge/v3/src/robo/robo.ts
--- a/Engine/Extras/RoboMerge/v3/src/robo/robo.ts
+++ b/Engine/Extras/RoboMerge/v3/src/robo/robo.ts
@@ -549,7 +549,7 @@ async function _initGraphBots(logger: ContextualLogger) {
 async function main(logger: ContextualLogger) {
 	while (true) {
 		try {
-			await initializePerforce(logger);
+			await initializePerforce(logger, vault);
 			break;
 		}
 		catch (err) {
diff --git a/Engine/Extras/RoboMerge/v3/src/robo/status-types.ts b/Engine/Extras/RoboMerge/v3/src/robo/status-types.ts
--- a/Engine/Extras/RoboMerge/v3/src/robo/status-types.ts
+++ b/Engine/Extras/RoboMerge/v3/src/robo/status-types.ts
@@ -18,6 +18,7 @@ type QueuedChange = Partial<ReconsiderArgs> & {
 
 export type FailureKind =
 	'Integration error' |
+	'Transfer error' |
 	'Exclusive check-out' |
 	'Merge conflict' |
 	'Commit failure' |
diff --git a/Engine/Extras/RoboMerge/v3/src/robo/vault.ts b/Engine/Extras/RoboMerge/v3/src/robo/vault.ts
--- a/Engine/Extras/RoboMerge/v3/src/robo/vault.ts
+++ b/Engine/Extras/RoboMerge/v3/src/robo/vault.ts
@@ -35,14 +35,19 @@ export class Vault {
     }
 
     get slackTokens() {
-        return this.valid ? this._vault['slack-tokens'] : null
+        return this._vault && this._vault['slack-tokens']
     }
 
     get cookieKey() {
-        return this._vault['cookie-key']
+        return this._vault && this._vault['cookie-key']
     }
 
     get ldapPassword() {
-        return this._vault['ldap-password']
+        return this._vault && this._vault['ldap-password']
+    }
+    
+    get perforceServers() { 
+        const servers = this._vault && this._vault['servers']
+        return servers ? new Map(Object.entries(servers)) : null
     }
 }
\ No newline at end of file
