Output long command errors to an appropriate logger

diff --git a/Engine/Extras/RoboMerge/v3/src/common/perforce.ts b/Engine/Extras/RoboMerge/v3/src/common/perforce.ts
--- a/Engine/Extras/RoboMerge/v3/src/common/perforce.ts
+++ b/Engine/Extras/RoboMerge/v3/src/common/perforce.ts
@@ -205,7 +205,7 @@ export function parseZTag(buffer: string, opts?: ExecZtagOpts) {
 }
 
 class CommandRecord {
-	constructor(public cmd: string, public start: Date = new Date()) {
+	constructor(public logger: ContextualLogger, public cmd: string, public start: Date = new Date()) {
 	}
 }
 
@@ -461,7 +461,7 @@ export async function initializePerforce(logger: ContextualLogger, vault: Vault)
 		for (const command of runningPerforceCommands) {
 			const durationMinutes = (now - command.start.valueOf()) / (60*1000)
 			if (durationMinutes > 10) {
-				logger.info(`Command still running after ${Math.round(durationMinutes)} minutes: ` + command.cmd)
+				command.logger.info(`Command still running after ${Math.round(durationMinutes)} minutes: ` + command.cmd)
 			}
 		}
 	}, 5*60*1000)
@@ -1769,7 +1769,7 @@ export class PerforceContext {
 		args = PerforceContext._getP4Cmd(roboWorkspace, args, optsIn)
 
 		// log what we're running
-		let cmd_rec = new CommandRecord('p4 ' + args.join(' '));
+		let cmd_rec = new CommandRecord(logger, 'p4 ' + args.join(' '));
 		// 	logger.verbose("Executing: " + cmd_rec.cmd);
 		runningPerforceCommands.add(cmd_rec);
 
