Don't fetch email for known bots

diff --git a/Engine/Extras/RoboMerge/v3/src/common/perforce.ts b/Engine/Extras/RoboMerge/v3/src/common/perforce.ts
--- a/Engine/Extras/RoboMerge/v3/src/common/perforce.ts
+++ b/Engine/Extras/RoboMerge/v3/src/common/perforce.ts
@@ -2,7 +2,7 @@
 import { execFile, ExecFileOptions } from 'child_process';
 import * as path from 'path';
 // should really separate out RM-specific operations, but no real gain at the moment
-import { postToRobomergeAlerts } from '../robo/notifications';
+import { isUserAKnownBot, postToRobomergeAlerts } from '../robo/notifications';
 import { roboAnalytics } from '../robo/roboanalytics';
 import { ContextualLogger, NpmLogLevel } from './logger';
 import { Vault } from '../robo/vault'
@@ -1511,7 +1511,7 @@ export class PerforceContext {
 	// get the email (according to P4) for a specific user
 	async getEmail(username: string) {
 		let m = null
-		if (username.length > 0 && !username.startsWith('@') && !username.startsWith('<')) {
+		if (username.length > 0 && !username.startsWith('@') && !username.startsWith('<') && !isUserAKnownBot(username)) {
 			const output = await this._execP4(null, ['user', '-o', username]);
 			// look for the email field
 			m = output.match(/\nEmail:\s+([^\n]+)\n/);
