If a password is supplied for a server in the vault then use that to login and refresh login ticket every hour

diff --git a/Engine/Extras/RoboMerge/v3/src/common/perforce.ts b/Engine/Extras/RoboMerge/v3/src/common/perforce.ts
--- a/Engine/Extras/RoboMerge/v3/src/common/perforce.ts
+++ b/Engine/Extras/RoboMerge/v3/src/common/perforce.ts
@@ -494,12 +494,19 @@ export class PerforceContext {
 		
 		const serverAddress = serverDetails && serverDetails.port
 
-		let loginArgs = ["login", "-s"]
+		let opts: ExecOpts = { serverAddress }
+		let loginArgs = ["login"]
 		if (serverDetails && serverDetails.user) {
 			loginArgs = ["-u", serverDetails.user, ...loginArgs]
 		}
+		if (serverDetails && serverDetails.password) {
+			opts.stdin = serverDetails.password
+		}
+		else {
+			loginArgs.push('-s')
+		}
 
-		const output = await PerforceContext.execAndParse(logger, null, loginArgs, { serverAddress });
+		const output = await PerforceContext.execAndParse(logger, null, loginArgs, opts);
 		let resp = output[0];
 	
 		if (resp && resp.User) {
@@ -523,6 +530,20 @@ export class PerforceContext {
 			}
 			this.servers.set(serverID, serverConfig)
 
+			if (serverDetails && serverDetails.password) {
+				// We're going to refresh our credentials every hour
+				setInterval(() =>  {
+					PerforceContext.execAndParse(logger, null, loginArgs, opts)
+					.then((output) => {
+						let resp = output[0];
+						if (!resp.User) {
+							logger.warn(`Failed to refresh credentials for ${serverID}: ${serverAddress}. Will try again in 1 hour.`)
+						}
+					})
+			
+				}, 60*60*1000)
+			}
+
 			if (serverAddress) {
 				logger.info(`P4 server ${serverID}: ${serverAddress} initialized. User=${resp.User}`)
 			}
