Catch when a channel fails to return channelInfo and report it to the alerts channel

diff --git a/Engine/Extras/RoboMerge/v3/src/robo/notifications.ts b/Engine/Extras/RoboMerge/v3/src/robo/notifications.ts
--- a/Engine/Extras/RoboMerge/v3/src/robo/notifications.ts
+++ b/Engine/Extras/RoboMerge/v3/src/robo/notifications.ts
@@ -400,7 +400,12 @@ export class SlackMessages {
 					else if (result.error === "user_is_restricted") {
 						// This an external user so we need to use the admin invite and validate that the channel ends in -ext
 						const channelInfo = await this.slack.getChannelInfo(channel)
-						if (channelInfo.channel.name.endsWith('-ext')) {
+						if (!channelInfo.channel) {
+							const errorMsg = `Failed to get name of channel ${channel} when inviting user ${emailAddress} (${user})`
+							this.smLogger.error(errorMsg)
+							postToRobomergeAlerts(errorMsg)
+						}
+						else if (channelInfo.channel.name.endsWith('-ext')) {
 							await this.addUserToChannel(emailAddress, channel, true)
 						}
 						else {
