Fix getChannelInfo returning different data

diff --git a/Engine/Extras/RoboMerge/v3/src/robo/slack.ts b/Engine/Extras/RoboMerge/v3/src/robo/slack.ts
--- a/Engine/Extras/RoboMerge/v3/src/robo/slack.ts
+++ b/Engine/Extras/RoboMerge/v3/src/robo/slack.ts
@@ -132,7 +132,7 @@ export class Slack {
 
 	async getChannelInfo(channel: string) {
 		try {
-			return (await Slack.slackApp.app.client.conversations.info({channel})).channel
+			return (await Slack.slackApp.app.client.conversations.info({channel}))
 		}
 		catch(error) {
 			this.logger.error(error)
