Focus block asset path to optimize sizes call

diff --git a/Engine/Extras/RoboMerge/v3/src/robo/nodebot.ts b/Engine/Extras/RoboMerge/v3/src/robo/nodebot.ts
--- a/Engine/Extras/RoboMerge/v3/src/robo/nodebot.ts
+++ b/Engine/Extras/RoboMerge/v3/src/robo/nodebot.ts
@@ -1682,8 +1682,8 @@ export class NodeBot extends PerforceStatefulBot implements NodeBotInterface {
 
 			if (blockAssetEdges.length > 0) {
 
-				let uassetSize = this.p4.sizes(`//....uasset@=${change.change}`, true)
-				let umapSize = this.p4.sizes(`//....umap@=${change.change}`, true)
+				let uassetSize = this.p4.sizes(`${this.branch.rootPath}.uasset@=${change.change}`, true)
+				let umapSize = this.p4.sizes(`${this.branch.rootPath}.umap@=${change.change}`, true)
 
 				const changeContainsAssets = ((await uassetSize)[0].fileCount > 0 || (await umapSize)[0].fileCount > 0)
 
