Fix edge server creation and cleaning

diff --git a/Engine/Extras/RoboMerge/v3/src/common/perforce.ts b/Engine/Extras/RoboMerge/v3/src/common/perforce.ts
--- a/Engine/Extras/RoboMerge/v3/src/common/perforce.ts
+++ b/Engine/Extras/RoboMerge/v3/src/common/perforce.ts
@@ -819,10 +819,10 @@ export class PerforceContext {
 		return this._execP4(null, ['-ztag', '-F', '%Address%', 'server', '-o', serverId])
 	}
 
-	async clean(roboWorkspace: RoboWorkspace) {
+	async clean(roboWorkspace: RoboWorkspace, edgeServerAddress?: string) {
 		let result;
 		try {
-			result = await this._execP4(roboWorkspace, ['clean', '-e', '-a']);
+			result = await this._execP4(roboWorkspace, ['clean', '-e', '-a'], {serverAddress: edgeServerAddress});
 		}
 		catch (reason) {
 			if (!isExecP4Error(reason)) {
@@ -885,7 +885,7 @@ export class PerforceContext {
 
 		let args = ['client', '-i']
 		if (edgeServer) {
-			args = ['-p', edgeServer.address, ...args, `--serverid=${edgeServer.id}`]
+			args.push(`--serverid=${edgeServer.id}`)
 		}
 
 		let workspaceForm = ''
@@ -904,7 +904,7 @@ export class PerforceContext {
 
 		// run the p4 client command
 		this.logger.info(`Executing: 'p4 client -i' to create workspace ${workspaceName}`);
-		return this._execP4(null, args, { stdin: workspaceForm, quiet: true });
+		return this._execP4(null, args, { stdin: workspaceForm, quiet: true, serverAddress: edgeServer?.address });
 	}
 
 	// Create a new workspace for Robomerge GraphBot
