Correct the parameter names on the admin invite of addUserToChannel

diff --git a/Engine/Extras/RoboMerge/v3/src/robo/slack.ts b/Engine/Extras/RoboMerge/v3/src/robo/slack.ts
--- a/Engine/Extras/RoboMerge/v3/src/robo/slack.ts
+++ b/Engine/Extras/RoboMerge/v3/src/robo/slack.ts
@@ -119,7 +119,7 @@ export class Slack {
 		try {
 			const slackApp = Slack.slackApp
 			if (externalUser) {
-				return await slackApp.app.client.admin.conversations.invite({token: slackApp.userToken || slackApp.botToken, channel, users})
+				return await slackApp.app.client.admin.conversations.invite({token: slackApp.userToken || slackApp.botToken, channel_id: channel, user_ids: users})
 			}
 			
 			return await slackApp.app.client.conversations.invite({token: slackApp.botToken, channel, users})
