Make sure to clean on correct edgeserver
Another tweak to ServerID for workspace creation

diff --git a/Engine/Extras/RoboMerge/v3/src/common/perforce.ts b/Engine/Extras/RoboMerge/v3/src/common/perforce.ts
--- a/Engine/Extras/RoboMerge/v3/src/common/perforce.ts
+++ b/Engine/Extras/RoboMerge/v3/src/common/perforce.ts
@@ -885,7 +885,7 @@ export class PerforceContext {
 
 		let args = ['client', '-i']
 		if (edgeServer) {
-			args.push(`--serverid=${edgeServer.id}`)
+			params.ServerID = edgeServer.id
 		}
 
 		let workspaceForm = ''
@@ -908,12 +908,12 @@ export class PerforceContext {
 	}
 
 	// Create a new workspace for Robomerge GraphBot
-	async newGraphBotWorkspace(name: string, extraParams: any, edgeServer?: EdgeServer) {
+	newGraphBotWorkspace(name: string, extraParams: any, edgeServer?: EdgeServer) {
 		return this.newWorkspace(name, {Root: getRootDirectoryForBranch(name), ...extraParams}, edgeServer);
 	}
 
 	// Create a new workspace for Robomerge to read branchspecs from
-	async newBranchSpecWorkspace(workspace: Workspace, bsDepotPath: string) {
+	newBranchSpecWorkspace(workspace: Workspace, bsDepotPath: string) {
 		let roots: string | string[] = '/app/data' // default linux path
 		if (workspace.directory !== roots) {
 			roots = [roots, workspace.directory] // specified directory
diff --git a/Engine/Extras/RoboMerge/v3/src/robo/edgebot.ts b/Engine/Extras/RoboMerge/v3/src/robo/edgebot.ts
--- a/Engine/Extras/RoboMerge/v3/src/robo/edgebot.ts
+++ b/Engine/Extras/RoboMerge/v3/src/robo/edgebot.ts
@@ -528,7 +528,7 @@ class EdgeBotImpl extends PerforceStatefulBot {
 				require('child_process').execSync(`rm -rf ${dir}/*`);
 			}
 			else {
-				await p4.clean(workspaceName);
+				await p4.clean(workspaceName, edgeServer?.address);
 			}
 		}
 		
