Fix track changes not working from Help page

diff --git a/Engine/Extras/RoboMerge/v3/public/js/boilerplate.js b/Engine/Extras/RoboMerge/v3/public/js/boilerplate.js
--- a/Engine/Extras/RoboMerge/v3/public/js/boilerplate.js
+++ b/Engine/Extras/RoboMerge/v3/public/js/boilerplate.js
@@ -47,6 +47,30 @@ function createErrorMessage(xhr, error) {
 		return `HTTP ${xhr.status}: ${error}`;
 }
 
+function promptFor(map) {
+	var result = {};
+	var success = true;
+	$.each(Object.keys(map), function(_, key) {
+		if (success) {
+			var p = map[key];
+			var dflt = "";
+			if (typeof(p) === "object") {
+				dflt = p.default || "";
+				p = p.prompt || key;
+			}
+
+			var data = prompt(p, dflt);
+			if (data === null) {
+				success = false;
+			}
+			else {
+				result[key] = data;
+			}
+		}
+	});
+	return success ? result : null;
+}
+
 // Add way to prompt the user on leave/reload before the perform an action
 function addOnBeforeUnload(message) {
 	$( window ).on("beforeunload", function() {
diff --git a/Engine/Extras/RoboMerge/v3/public/js/main.js b/Engine/Extras/RoboMerge/v3/public/js/main.js
--- a/Engine/Extras/RoboMerge/v3/public/js/main.js
+++ b/Engine/Extras/RoboMerge/v3/public/js/main.js
@@ -1,30 +1,6 @@
 // Copyright Epic Games, Inc. All Rights Reserved.
 "use strict"
 
-function promptFor(map) {
-	var result = {};
-	var success = true;
-	$.each(Object.keys(map), function(_, key) {
-		if (success) {
-			var p = map[key];
-			var dflt = "";
-			if (typeof(p) === "object") {
-				dflt = p.default || "";
-				p = p.prompt || key;
-			}
-
-			var data = prompt(p, dflt);
-			if (data === null) {
-				success = false;
-			}
-			else {
-				result[key] = data;
-			}
-		}
-	});
-	return success ? result : null;
-}
-
 class BotPips {
 	constructor() {
 		this.blocked = 0;
