Fix internal server error when tracking changes for certain paths

diff --git a/Engine/Extras/RoboMerge/v3/src/robo/ipc.ts b/Engine/Extras/RoboMerge/v3/src/robo/ipc.ts
--- a/Engine/Extras/RoboMerge/v3/src/robo/ipc.ts
+++ b/Engine/Extras/RoboMerge/v3/src/robo/ipc.ts
@@ -215,7 +215,11 @@ export class IPC {
 			if (desc.path) {
 				const stream = getStreamFromPath(desc.path)
 				if (stream) {
-					return graph.findNodeForStream(stream)
+					try {
+						return graph.findNodeForStream(stream)
+					}
+					catch {						
+					}
 				}
 			}
 			return null
