Change search for outstanding changes to merge to use [REDACTED_HANDLE],now instead of [at]>since

diff --git a/Engine/Extras/RoboMerge/v3/src/common/perforce.ts b/Engine/Extras/RoboMerge/v3/src/common/perforce.ts
--- a/Engine/Extras/RoboMerge/v3/src/common/perforce.ts
+++ b/Engine/Extras/RoboMerge/v3/src/common/perforce.ts
@@ -614,7 +614,7 @@ export class PerforceContext {
 	 * @return Promise to list of changelists
 	 */
 	changes(path_in: string, since: number, limit?: number, status?: ChangelistStatus, quiet: boolean = true): Promise<Change[]> {
-		const path = since > 0 ? path_in + '@>' + since : path_in;
+		const path = since > 0 ? `${path_in}@${since},now` : path_in;
 		const args = ['changes', '-l',
 			(status ? `-s${status}` : '-ssubmitted'),
 			...(limit ? [`-m${limit}`] : []),
