Pick a workspace on the correct server when reconsidering across servers

diff --git a/Engine/Extras/RoboMerge/v3/src/robo/nodebot.ts b/Engine/Extras/RoboMerge/v3/src/robo/nodebot.ts
--- a/Engine/Extras/RoboMerge/v3/src/robo/nodebot.ts
+++ b/Engine/Extras/RoboMerge/v3/src/robo/nodebot.ts
@@ -1606,11 +1606,13 @@ export class NodeBot extends PerforceStatefulBot implements NodeBotInterface {
 		// targets as long as all the target workspaces are on the same edge, but it will still crash if you
 		// have multiple targets with workspaces on different edgeservers
 		const numTargets = (result.info.targets || []).length;
-		let isManualChange = numTargets > 0 ? result.info.targets![0].flags.has('manual') : false;
+		const isManualChange = numTargets > 0 ? result.info.targets![0].flags.has('manual') : false;
 		if (change.forceCreateAShelf || (change.isUserRequest && !change.forceStompChanges) || isManualChange) {
 			if (!optWorkspaceOverride && numTargets == 1) {
-				const targetStream = result.info.targets![0].branch.stream ? result.info.targets![0].branch.stream.toLowerCase() : undefined
-				const targetWorkspaceDef = await p4util.chooseBestWorkspaceForUser(this.p4, result.info.owner||result.info.author, targetStream)
+				const target = result.info.targets![0]
+				const targetServerP4 = PerforceContext.getServerContext(this.logger, target.branch.config.streamServer)
+				const targetStream = target.branch.stream ? target.branch.stream.toLowerCase() : undefined
+				const targetWorkspaceDef = await p4util.chooseBestWorkspaceForUser(targetServerP4, result.info.owner||result.info.author, targetStream)
 				if (targetWorkspaceDef) {
 					result.info.targetWorkspaceForShelf = targetWorkspaceDef.client
 					if (targetWorkspaceDef.Stream) {
@@ -1626,8 +1628,14 @@ export class NodeBot extends PerforceStatefulBot implements NodeBotInterface {
 				}
 			}
 			if (optWorkspaceOverride) {
+				if (numTargets > 1) {	
+					const emResult = new EdgeMergeResults()
+					emResult.errors = ["Specifying a workspace override with multiple targets does not make sense"]
+					return emResult
+				}
+				const targetServerP4 = numTargets > 0 ? PerforceContext.getServerContext(this.logger, result.info.targets![0].branch.config.streamServer) : this.p4
 				// see if we need to talk to an edge server to create the shelf
-				const edgeServer = await this.p4.getWorkspaceEdgeServer(
+				const edgeServer = await targetServerP4.getWorkspaceEdgeServer(
 					coercePerforceWorkspace(optWorkspaceOverride)!.name)
 				if (edgeServer) {
 					result.info.edgeServerToHostShelf = edgeServer
