Support narrowing transfered file view using stream subpath

diff --git a/Engine/Extras/RoboMerge/v3/src/common/p4util.ts b/Engine/Extras/RoboMerge/v3/src/common/p4util.ts
--- a/Engine/Extras/RoboMerge/v3/src/common/p4util.ts
+++ b/Engine/Extras/RoboMerge/v3/src/common/p4util.ts
@@ -4,6 +4,14 @@ import { _setTimeout } from './helper';
 import { ContextualLogger } from './logger';
 import { Change, ClientSpec, PerforceContext, Workspace } from './perforce';
 
+
+export function makeRegexFromPerforcePath(path: string) {
+	path = path.replace(/[/\-\\^$*+?.()|[\]{}]/g, '\\$&') // re.escape
+	path = path.replace('\\*', '[^/]*')
+	path = path.replace('\\.\\.\\.', '.*')
+	return new RegExp(path)
+}
+
 const USER_WORKSPACE_EXCLUDE_PATTERNS: (RegExp | string)[] = [
 	'horde-p4bridge-',
 	'swarm-'
diff --git a/Engine/Extras/RoboMerge/v3/src/common/perforce.ts b/Engine/Extras/RoboMerge/v3/src/common/perforce.ts
--- a/Engine/Extras/RoboMerge/v3/src/common/perforce.ts
+++ b/Engine/Extras/RoboMerge/v3/src/common/perforce.ts
@@ -7,6 +7,7 @@ import { roboAnalytics } from '../robo/roboanalytics';
 import { ContextualLogger, NpmLogLevel } from './logger';
 import { Vault } from '../robo/vault'
 import { VersionReader } from './version';
+import { makeRegexFromPerforcePath } from './p4util';
 
 type Type = 'string' | 'integer' | 'boolean'
 
@@ -986,15 +987,25 @@ export class PerforceContext {
 	async transfer(targetP4: PerforceContext, roboWorkspace: RoboWorkspace, source: IntegrationSource, dest_changelist: number, target: IntegrationTarget)
 	: Promise<string[]> {
 
-		const desc = await this.describe(source.changelist)
+		// TODO: support branchspecs
+		if (source.branchspec) {
+			return ["Branchspecs not yet supported for transfers."]
+		}
+
+		const files = await this.files(`${source.path_from}@=${source.changelist}`)
 
 		let failures: string[] = []
 
 		const sourceStreamLength = source.stream!.length
+		const targetFilter = `${target.stream!}/...` == target.path_to ? null : makeRegexFromPerforcePath(target.path_to)
 
-		await Promise.all(desc.entries.map(entry => {
-			const targetDepotFile = target.stream! + entry.depotFile.substring(sourceStreamLength)
-			switch(entry.action) {
+		await Promise.all(files.map(file => {
+			const targetDepotFile = target.stream! + file.depotFile.substring(sourceStreamLength)
+			if (targetFilter && !targetDepotFile.match(targetFilter)) {
+				// filter out files that don't match the target.path_to
+				return
+			}
+			switch(file.action) {
 				case 'add':
 				case 'branch': {
 					// Validate file doesn't exist
@@ -1002,8 +1013,8 @@ export class PerforceContext {
 					.then(fstat => {
 						if (fstat.length == 0 || fstat[0].headAction.includes('delete')) {
 							return targetP4.where(roboWorkspace, targetDepotFile)
-							.then(where => { return this.print(`${entry.depotFile}@${source.changelist}`, where[0].path) })
-							.then(() => { return targetP4.add(roboWorkspace, dest_changelist, targetDepotFile, entry.type) })
+							.then(where => { return this.print(`${file.depotFile}@${source.changelist}`, where[0].path) })
+							.then(() => { return targetP4.add(roboWorkspace, dest_changelist, targetDepotFile, file.type) })
 						}
 						else {
 							failures.push(`File ${targetDepotFile} already exists when being added in CL# ${source.changelist}`)
@@ -1018,8 +1029,8 @@ export class PerforceContext {
 					.then(fstat => {
 						if (fstat.length > 0 && !fstat[0].headAction.includes('delete')) {
 							return targetP4.sync(roboWorkspace, targetDepotFile, {opts: ['-k']})
-							.then(() => { return targetP4.edit(roboWorkspace, dest_changelist, targetDepotFile, ['-k', '-t', entry.type]) })
-							.then(() => { return this.print(`${entry.depotFile}@${source.changelist}`, fstat[0].clientFile) })	
+							.then(() => { return targetP4.edit(roboWorkspace, dest_changelist, targetDepotFile, ['-k', '-t', file.type]) })
+							.then(() => { return this.print(`${file.depotFile}@${source.changelist}`, fstat[0].clientFile) })	
 						}
 						else {
 							failures.push(`File ${targetDepotFile} does not exist when being edited in CL# ${source.changelist}`)
@@ -1031,7 +1042,7 @@ export class PerforceContext {
 					return targetP4.delete(roboWorkspace, dest_changelist, targetDepotFile)
 				}
 				case 'move/delete': {
-					return this.integrated(null, entry.depotFile, {intoOnly: true, startCL: source.changelist})
+					return this.integrated(null, file.depotFile, {intoOnly: true, startCL: source.changelist})
 					.then(integrated => {
 						const integ = integrated.find(integ => parseInt(integ.change) == source.changelist)
 						if (integ) {
@@ -1042,7 +1053,7 @@ export class PerforceContext {
 							.then(move => {	return this.print(`${integ.fromFile}@${source.changelist}`, move[0].path) })
 						}
 						else {
-							failures.push(`Unable to look up integrated record for ${entry.depotFile} in CL# ${source.changelist}`)
+							failures.push(`Unable to look up integrated record for ${file.depotFile} in CL# ${source.changelist}`)
 							return
 						}
 					})
@@ -1050,7 +1061,7 @@ export class PerforceContext {
 				case 'move/add': return // move/add is handled from the move/delete portion of the pair
 
 				default: {
-					failures.push(`Unhandled transfer action ${entry.action} for ${entry.depotFile} in CL# ${source.changelist}`)
+					failures.push(`Unhandled transfer action ${file.action} for ${file.depotFile} in CL# ${source.changelist}`)
 					return
 				}
 			}
