Don't output error message when no changes are found in a monitored stream

diff --git a/Engine/Extras/RoboMerge/v3/src/robo/nodebot.ts b/Engine/Extras/RoboMerge/v3/src/robo/nodebot.ts
--- a/Engine/Extras/RoboMerge/v3/src/robo/nodebot.ts
+++ b/Engine/Extras/RoboMerge/v3/src/robo/nodebot.ts
@@ -222,7 +222,9 @@ export class NodeBot extends PerforceStatefulBot implements NodeBotInterface {
 			}
 		}
 		catch (err) {
-			this.nodeBotLogger.printException(err, `${this.fullName} Error while querying P4 for changes`)
+			if (err.message !== "Expected exactly one change. Got 0") {
+				this.nodeBotLogger.printException(err, `${this.fullName} Error while querying P4 for changes`)
+			}
 		}
 
 		return []
