Can specify a pattern that if it exists in the description the change should be skipped

diff --git a/Engine/Extras/RoboMerge/v3/README.md b/Engine/Extras/RoboMerge/v3/README.md
--- a/Engine/Extras/RoboMerge/v3/README.md
+++ b/Engine/Extras/RoboMerge/v3/README.md
@@ -106,6 +106,7 @@ for both node and edge.
 |`disallowSkip`             |e |            | |Remove skip option from UI               | |
 |`emailOnBlockage`          |e |`true`      | |Email owners of conflicts?               | |
 |`excludeAuthors`           |e |            | |Skip changes by these authors, e.g. skip `buildmachine` commits| |
+|`excludeDescriptions`      |e |            | |Skip changes that contain these patterns in their description| |
 |`ignoreBranchspecs`        |e |            | |Where branchspecs specified, ignore for this edge| |
 |`incognitoMode`            |e |`false`     | |Terse description in committed changelists| |
 |`initialCL`                |e |            | |First run only: which CL start _after_   | |
diff --git a/Engine/Extras/RoboMerge/v3/functional_tests/framework/src/framework.ts b/Engine/Extras/RoboMerge/v3/functional_tests/framework/src/framework.ts
--- a/Engine/Extras/RoboMerge/v3/functional_tests/framework/src/framework.ts
+++ b/Engine/Extras/RoboMerge/v3/functional_tests/framework/src/framework.ts
@@ -169,6 +169,7 @@ type RobomergeBranchOptions = {
 	incognitoMode: boolean
 
 	excludeAuthors: string[] // if present, completely overrides BotConfig
+	excludeDescriptions: string[] // if present, completely overrides BotConfig
 }
 
 export type RobomergeBranchSpec = Partial<RobomergeBranchOptions> & {
@@ -194,6 +195,7 @@ type EdgeOptionFields = {
 	terminal: boolean // changes go along terminal edges but no further
 
 	excludeAuthors: string[]
+	excludeDescriptions: string[]
 
 	// by default, specify when gate catch ups are allowed; can be inverted to disallow
 	integrationWindow: IntegrationWindowPane[]
@@ -1107,8 +1109,8 @@ export class P4Util {
 			.then(() => this.submit(client, `Adding file ${name}`))
 	}
 
-	static editFileAndSubmit(client: P4Client, name: string, newContent: string, robomergeCommand?: string) {
-		let commitMessage = `Edited file '${name}'`
+	static editFileAndSubmit(client: P4Client, name: string, newContent: string, robomergeCommand?: string, commitMessage?: string) {
+		commitMessage = commitMessage || `Edited file '${name}'`
 		if (robomergeCommand) {
 			commitMessage += '\n#robomerge ' + robomergeCommand
 		}
diff --git a/Engine/Extras/RoboMerge/v3/functional_tests/framework/src/tests.ts b/Engine/Extras/RoboMerge/v3/functional_tests/framework/src/tests.ts
--- a/Engine/Extras/RoboMerge/v3/functional_tests/framework/src/tests.ts
+++ b/Engine/Extras/RoboMerge/v3/functional_tests/framework/src/tests.ts
@@ -17,6 +17,8 @@ import { ExclusiveCheckout } from './tests/exclusive-checkout'
 import { ExclusiveCheckoutUnlock } from './tests/exclusive-checkout-unlock'
 import { ExcludeAuthors } from './tests/exclude-authors'
 import { ExcludeAuthorsPerEdge } from './tests/exclude-authors-per-edge'
+import { ExcludeDescriptions } from './tests/exclude-descriptions'
+import { ExcludeDescriptionsPerEdge } from './tests/exclude-descriptions-per-edge'
 import { ForwardCommands, ForwardCommands2 } from './tests/forward-commands'
 import { Ignore } from './tests/ignore'
 import { ImplicitCommands } from './tests/implicit-commands'
@@ -219,7 +221,9 @@ async function go() {
 
 		new PostToAdditionalChannel(p4),
 		new BranchspecTest(p4),
-		new ExclusiveCheckoutUnlock(p4)
+		new ExclusiveCheckoutUnlock(p4),
+		new ExcludeDescriptions(p4),
+		new ExcludeDescriptionsPerEdge(p4) // 50
 	]
 
 	const TARGET_TEST_DEFS: [string[], string | null][] = [
diff --git a/Engine/Extras/RoboMerge/v3/functional_tests/framework/src/tests/exclude-descriptions-per-edge.ts b/Engine/Extras/RoboMerge/v3/functional_tests/framework/src/tests/exclude-descriptions-per-edge.ts
new file mode 100644
--- /dev/null
+++ b/Engine/Extras/RoboMerge/v3/functional_tests/framework/src/tests/exclude-descriptions-per-edge.ts
@@ -0,0 +1,59 @@
+// Copyright Epic Games, Inc. All Rights Reserved.
+import { EdgeProperties, FunctionalTest, P4Util, Stream } from '../framework'
+
+const streams: Stream[] =
+	[ {name: 'Main', streamType: 'mainline'}
+	, {name: 'Dev-Perkin', streamType: 'development', parent: 'Main'}
+	, {name: 'Dev-Pootle', streamType: 'development', parent: 'Main'}
+	]
+
+export class ExcludeDescriptionsPerEdge extends FunctionalTest {
+	async setup() {
+		await this.p4.depot('stream', this.depotSpec())
+		await this.createStreamsAndWorkspaces(streams)
+
+		await P4Util.addFileAndSubmit(this.getClient('Main'), 'test.txt', 'Initial content')
+
+		const desc = 'Initial populate'
+
+		await Promise.all([
+			this.p4.populate(this.getStreamPath('Dev-Perkin'), desc),
+			this.p4.populate(this.getStreamPath('Dev-Pootle'), desc)
+		])
+	}
+
+	async run() {
+
+        let client = this.getClient('Main')
+		let content = 'Initial content'
+		content += '\n\nsubmit1 content'
+		await P4Util.editFileAndSubmit(client, 'test.txt', content, undefined, "Skipped Submit")
+
+		content += '\n\nsubmit2 content'
+		await P4Util.editFileAndSubmit(client, 'test.txt', content, undefined, "Allowed Submit")
+	}
+
+	verify() {
+		return Promise.all([
+			this.checkHeadRevision('Main', 'test.txt', 3),
+			this.checkHeadRevision('Dev-Perkin', 'test.txt', 3),
+			this.ensureBlocked('Main', 'Dev-Pootle') // skipped submit1 does try to merge submit2 creating a conflict
+		])
+	}
+
+	getBranches() {
+		return [ 
+		    this.makeForceAllBranchDef('Main', ['Dev-Perkin', 'Dev-Pootle'])
+		  , this.makeForceAllBranchDef('Dev-Perkin', [])
+		  , this.makeForceAllBranchDef('Dev-Pootle', [])
+		]
+	}
+
+	getEdges(): EdgeProperties[] {
+		return [
+		  { from: this.fullBranchName('Main'), to: this.fullBranchName('Dev-Pootle')
+		  , excludeDescriptions: ['S.* S.*']
+		  }
+		]
+	}
+}
diff --git a/Engine/Extras/RoboMerge/v3/functional_tests/framework/src/tests/exclude-descriptions.ts b/Engine/Extras/RoboMerge/v3/functional_tests/framework/src/tests/exclude-descriptions.ts
new file mode 100644
--- /dev/null
+++ b/Engine/Extras/RoboMerge/v3/functional_tests/framework/src/tests/exclude-descriptions.ts
@@ -0,0 +1,52 @@
+// Copyright Epic Games, Inc. All Rights Reserved.
+import { FunctionalTest, P4Util, Stream } from '../framework'
+
+const streams: Stream[] =
+	[ {name: 'Main', streamType: 'mainline'}
+	, {name: 'Dev-Perkin', streamType: 'development', parent: 'Main'}
+	, {name: 'Dev-Pootle', streamType: 'development', parent: 'Main'}
+	]
+
+export class ExcludeDescriptions extends FunctionalTest {
+	async setup() {
+		await this.p4.depot('stream', this.depotSpec())
+		await this.createStreamsAndWorkspaces(streams)
+
+		await P4Util.addFileAndSubmit(this.getClient('Main'), 'test.txt', 'Initial content')
+
+		const desc = 'Initial populate'
+
+		await Promise.all([
+			this.p4.populate(this.getStreamPath('Dev-Perkin'), desc),
+			this.p4.populate(this.getStreamPath('Dev-Pootle'), desc)
+		])
+	}
+
+	async run() {
+        let client = this.getClient('Main')
+		let content = 'Initial content'
+		content += '\n\nsubmit1 content'
+		await P4Util.editFileAndSubmit(client, 'test.txt', content, undefined, "Allowed Submit")
+
+		content += '\n\nsubmit2 content'
+		await P4Util.editFileAndSubmit(client, 'test.txt', content, undefined, "Skipped Submit")
+	}
+
+	verify() {
+		return Promise.all([
+			this.checkHeadRevision('Main', 'test.txt', 3),
+			this.checkHeadRevision('Dev-Perkin', 'test.txt', 2),
+			this.checkHeadRevision('Dev-Pootle', 'test.txt', 2)
+		])
+	}
+
+	getBranches() {
+		const mainBranchDef = this.makeForceAllBranchDef('Main', ['Dev-Perkin', 'Dev-Pootle'])
+		mainBranchDef.excludeDescriptions = ['S.* S.*']
+		return [ 
+		    mainBranchDef
+		  , this.makeForceAllBranchDef('Dev-Perkin', [])
+		  , this.makeForceAllBranchDef('Dev-Pootle', [])
+		]
+	}
+}
diff --git a/Engine/Extras/RoboMerge/v3/src/robo/branchdefs.ts b/Engine/Extras/RoboMerge/v3/src/robo/branchdefs.ts
--- a/Engine/Extras/RoboMerge/v3/src/robo/branchdefs.ts
+++ b/Engine/Extras/RoboMerge/v3/src/robo/branchdefs.ts
@@ -20,6 +20,7 @@ export interface BotConfig {
 	triager: string | null
 	checkIntervalSecs: number
 	excludeAuthors: string[]
+	excludeDescriptions: string[]
 	emailOnBlockage: boolean
 	visibility: string[] | string
 	slackChannel: string
@@ -93,6 +94,7 @@ export const commonOptionFieldsPrototype = {
 	incognitoMode: false,
 
 	excludeAuthors: [] as string[], // if present, completely overrides BotConfig
+	excludeDescriptions: [] as string[], // if present, completely overrides BotConfig
 
 	// by default, specify when gate catch ups are allowed; can be inverted to disallow
 	integrationWindow: [] as IntegrationWindowPane[],
@@ -280,6 +282,7 @@ export class BranchDefs {
 			emailOnBlockage: true,
 			checkIntervalSecs: 30.0,
 			excludeAuthors: [],
+			excludeDescriptions: [],
 			visibility: ['fte'],
 			slackChannel: '',
 			reportToBuildHealth: false,
diff --git a/Engine/Extras/RoboMerge/v3/src/robo/branchgraph.ts b/Engine/Extras/RoboMerge/v3/src/robo/branchgraph.ts
--- a/Engine/Extras/RoboMerge/v3/src/robo/branchgraph.ts
+++ b/Engine/Extras/RoboMerge/v3/src/robo/branchgraph.ts
@@ -54,6 +54,7 @@ if (botname === '__TEST__') {
 	private propagateBotPropertiesToNodes() {
 		const propertiesToPropagate: [keyof BotConfig, keyof NodeOptions, ConfigBlendMode | null][] = [
 			['excludeAuthors', 'excludeAuthors', 'override'],
+			['excludeDescriptions', 'excludeDescriptions', 'override'],
 			['badgeUrlOverride', 'badgeUrlOverride', 'override']
 		]
 
@@ -88,6 +89,7 @@ if (botname === '__TEST__') {
 			['disallowSkip', 'disallowSkip', null],
 			['incognitoMode', 'incognitoMode', null],
 			['excludeAuthors', 'excludeAuthors', 'override'],
+			['excludeDescriptions', 'excludeDescriptions', 'override'],
 			['integrationWindow', 'integrationWindow', null],
 			['invertIntegrationWindow', 'invertIntegrationWindow', null],
 		]
diff --git a/Engine/Extras/RoboMerge/v3/src/robo/edgebot.ts b/Engine/Extras/RoboMerge/v3/src/robo/edgebot.ts
--- a/Engine/Extras/RoboMerge/v3/src/robo/edgebot.ts
+++ b/Engine/Extras/RoboMerge/v3/src/robo/edgebot.ts
@@ -170,6 +170,9 @@ class EdgeBotImpl extends PerforceStatefulBot {
 	get excludedAuthors() {
 		return this.options.excludeAuthors || []
 	}
+	get excludedDescriptions() {
+		return this.options.excludeDescriptions || []
+	}
 	protected get logger() {
 		return this.edgeBotLogger
 	}
@@ -1076,6 +1079,7 @@ class EdgeBotImpl extends PerforceStatefulBot {
 		status.disallowSkip = this.options.disallowSkip
 		status.incognitoMode = this.options.incognitoMode
 		status.excludeAuthors = this.options.excludeAuthors
+		status.excludeDescriptions = this.options.excludeDescriptions
 
 		return status as EdgeStatusFields
 	}
diff --git a/Engine/Extras/RoboMerge/v3/src/robo/nodebot.ts b/Engine/Extras/RoboMerge/v3/src/robo/nodebot.ts
--- a/Engine/Extras/RoboMerge/v3/src/robo/nodebot.ts
+++ b/Engine/Extras/RoboMerge/v3/src/robo/nodebot.ts
@@ -1951,14 +1951,28 @@ export class NodeBot extends PerforceStatefulBot implements NodeBotInterface {
 							// let ignore and disregard cancel each other out here, but don't honour ignore here so we can validate it at the bottom of the function
 							const authorIsExcluded = edgeProperties && edgeProperties.excludeAuthors && edgeProperties.excludeAuthors.indexOf(info.author) >= 0
 
-							if (!authorIsExcluded || (flags.has('disregardexcludedauthors') && !flags.has('ignore'))) {
-								mergeActions.push(action)
-							}
-							else {
+							if (authorIsExcluded && (!flags.has('disregardexcludedauthors') || flags.has('ignore'))) {
 								// skip this one
 								const skipMessage = `Skipping CL ${info.cl} due to excluded author '${info.author}'`
 								this.nodeBotLogger.info(skipMessage)
+								continue
+							}
 
+							let descriptionIsExcluded = false
+							if (edgeProperties && edgeProperties.excludeDescriptions) {
+								for (const pattern of edgeProperties.excludeDescriptions) {
+									if (info.description.search(pattern) >= 0) {
+										//skip this one
+										const skipMessage = `Skipping CL ${info.cl} due to excluded description pattern '${pattern}'`
+										this.nodeBotLogger.info(skipMessage)
+										descriptionIsExcluded = true
+										break
+									}
+								}
+							}
+
+							if (!descriptionIsExcluded) { //TODO: consider disregard flag
+								mergeActions.push(action)
 							}
 						}
 
diff --git a/Engine/Extras/RoboMerge/v3/src/robo/status-types.ts b/Engine/Extras/RoboMerge/v3/src/robo/status-types.ts
--- a/Engine/Extras/RoboMerge/v3/src/robo/status-types.ts
+++ b/Engine/Extras/RoboMerge/v3/src/robo/status-types.ts
@@ -114,6 +114,7 @@ export type EdgeStatusFields = BotStatusFields & {
 	disallowSkip: boolean
 	incognitoMode: boolean
 	excludeAuthors: string[]
+	excludeDescriptions: string[]
 
 	lastGoodCL?: number
 	lastGoodCLJobLink?: string
