Make description exclusion work correctly across multiple lines

diff --git a/Engine/Extras/RoboMerge/v3/src/robo/nodebot.ts b/Engine/Extras/RoboMerge/v3/src/robo/nodebot.ts
--- a/Engine/Extras/RoboMerge/v3/src/robo/nodebot.ts
+++ b/Engine/Extras/RoboMerge/v3/src/robo/nodebot.ts
@@ -1961,7 +1961,7 @@ export class NodeBot extends PerforceStatefulBot implements NodeBotInterface {
 							let descriptionIsExcluded = false
 							if (edgeProperties && edgeProperties.excludeDescriptions) {
 								for (const pattern of edgeProperties.excludeDescriptions) {
-									if (info.description.search(pattern) >= 0) {
+									if (info.description.search(new RegExp(pattern,"s")) >= 0) {
 										//skip this one
 										const skipMessage = `Skipping CL ${info.cl} due to excluded description pattern '${pattern}'`
 										this.nodeBotLogger.info(skipMessage)
