Fix type of commitmessage

diff --git a/Engine/Extras/RoboMerge/v3/functional_tests/framework/src/framework.ts b/Engine/Extras/RoboMerge/v3/functional_tests/framework/src/framework.ts
--- a/Engine/Extras/RoboMerge/v3/functional_tests/framework/src/framework.ts
+++ b/Engine/Extras/RoboMerge/v3/functional_tests/framework/src/framework.ts
@@ -1109,13 +1109,13 @@ export class P4Util {
 			.then(() => this.submit(client, `Adding file ${name}`))
 	}
 
-	static editFileAndSubmit(client: P4Client, name: string, newContent: string, robomergeCommand?: string, commitMessage?: string) {
-		commitMessage = commitMessage || `Edited file '${name}'`
+	static editFileAndSubmit(client: P4Client, name: string, newContent: string, robomergeCommand?: string, submitMessage?: string) {
+		let commitMessage = submitMessage || `Edited file '${name}'`
 		if (robomergeCommand) {
 			commitMessage += '\n#robomerge ' + robomergeCommand
 		}
 		return this.editFile(client, name, newContent)
-			.then(() => this.submit(client, commitMessage))
+			.then(() => this.submit(client, commitMessage!))
 	}
 
 	static deleteFileAndSubmit(client: P4Client, name: string, robomergeCommand?: string) {
