Fix workspaces being constantly recleaned

diff --git a/Engine/Extras/RoboMerge/v3/src/common/p4util.ts b/Engine/Extras/RoboMerge/v3/src/common/p4util.ts
--- a/Engine/Extras/RoboMerge/v3/src/common/p4util.ts
+++ b/Engine/Extras/RoboMerge/v3/src/common/p4util.ts
@@ -17,16 +17,16 @@ const USER_WORKSPACE_EXCLUDE_PATTERNS: (RegExp | string)[] = [
 	'swarm-'
 ]
 
-type WorkspaceID = {
-	serverAddress?: string
-	workspaceName: string
+type WorkspaceID = string
+function getWorkspaceID(serverAddress: string | undefined, workspaceName: string) {
+	return JSON.stringify({serverAddress, workspaceName})
 }
 // true has been cleaned, false is being cleaned, undefined needs cleaning
 let cleanedWorkspaces: Map<WorkspaceID, boolean> = new Map()
 
 export async function cleanWorkspace(logger: ContextualLogger, p4: PerforceContext, workspace: Workspace, workspaceRoot: string, edgeServerAddress?: string) {
 	
-	const workspaceID: WorkspaceID = {serverAddress: p4.serverAddress, workspaceName: workspace.name}
+	const workspaceID: WorkspaceID = getWorkspaceID(p4.serverAddress, workspace.name)
 
 	// if the workspace is being cleaned, wait for it to finish
 	let workspaceCleaned = cleanedWorkspaces.get(workspaceID)
