Ability to specify a customized changelist parser for a given node or bot

diff --git a/Engine/Extras/RoboMerge/v3/src/robo/branchdefs.ts b/Engine/Extras/RoboMerge/v3/src/robo/branchdefs.ts
--- a/Engine/Extras/RoboMerge/v3/src/robo/branchdefs.ts
+++ b/Engine/Extras/RoboMerge/v3/src/robo/branchdefs.ts
@@ -17,6 +17,7 @@ export interface BotConfig {
 	nagSchedule: number[] | null
 	nagAcknowledgedSchedule: number[] | null
 	nagAcknowledgedLeeway: number | null
+	changelistParser: string | null
 	triager: string | null
 	checkIntervalSecs: number
 	excludeAuthors: string[]
@@ -46,6 +47,7 @@ const branchBasePrototype = {
 	forceFlowTo: [''],
 	defaultFlow: [''],
 	macros: {} as { [name: string]: string[] } | undefined,
+	changelistParser: null as string | null,
 	resolver: '' as string | null,
 	triager: '' as string | null,
 	nagWhenBlocked: false as boolean | null,
@@ -274,6 +276,7 @@ export class BranchDefs {
 			isDefaultBot: false,
 			noStreamAliases: false,
 			globalNotify: [],
+			changelistParser: null,
 			triager: null,
 			nagSchedule: null,
 			nagAcknowledgedSchedule: null,
diff --git a/Engine/Extras/RoboMerge/v3/src/robo/branchgraph.ts b/Engine/Extras/RoboMerge/v3/src/robo/branchgraph.ts
--- a/Engine/Extras/RoboMerge/v3/src/robo/branchgraph.ts
+++ b/Engine/Extras/RoboMerge/v3/src/robo/branchgraph.ts
@@ -210,8 +210,8 @@ if (botname === '__TEST__') {
 		this.finish()
 	}
 
-	private static _getBoolConfig(fromOptions?: boolean, fromConfig?: boolean): boolean {
-		return !!(fromOptions === undefined ? fromConfig : fromOptions)
+	private static _getConfig<Type>(fromOptions?: Type, fromConfig?: Type): Type | undefined {
+		return (fromOptions === undefined ? fromConfig : fromOptions)
 	}
 
 	private add(name: string, options: NodeOptions) {
@@ -231,8 +231,8 @@ if (botname === '__TEST__') {
 			rootPath: options.rootPath || "",
 			uniqueBranch: options.uniqueBranch || false,
 			badgeProject: options.badgeProject || null,
-			isDefaultBot: BranchGraph._getBoolConfig(options.isDefaultBot, this.config.isDefaultBot),
-			emailOnBlockage: BranchGraph._getBoolConfig(options.emailOnBlockage, this.config.emailOnBlockage),
+			isDefaultBot: !!BranchGraph._getConfig(options.isDefaultBot, this.config.isDefaultBot),
+			emailOnBlockage: !!BranchGraph._getConfig(options.emailOnBlockage, this.config.emailOnBlockage),
 			flowsTo: options.flowsTo || [],
 			notify: (options.notify || []).concat(this.config.globalNotify),
 			forceFlowTo: [], // will compute
@@ -241,6 +241,7 @@ if (botname === '__TEST__') {
 			blockAssetTargets: new Set<string>(),
 			enabled: !options.disabled,
 			allowDeadend: !options.disallowDeadend,
+			changelistParser: BranchGraph._getConfig(options.changelistParser, this.config.changelistParser) || null,
 			resolver: options.resolver || null,
 			triager: options.triager || null,
 			nagSchedule: options.nagSchedule || null,
@@ -251,6 +252,9 @@ if (botname === '__TEST__') {
 			visibility: options.visibility || this.config.visibility,
 			get isMonitored() { return !!(this.bot && this.bot.isRunning) }
 		}
+		if (branch.changelistParser) {
+			import(`./changelistparsers/${branch.changelistParser}`)
+		}
 
 		// non-enumerable properties so they don't get logged
 		Object.defineProperty(branch, 'bot', { enumerable: false, value: null, writable: true })
diff --git a/Engine/Extras/RoboMerge/v3/src/robo/changelistparsers/customtestparser.ts b/Engine/Extras/RoboMerge/v3/src/robo/changelistparsers/customtestparser.ts
new file mode 100644
--- /dev/null
+++ b/Engine/Extras/RoboMerge/v3/src/robo/changelistparsers/customtestparser.ts
@@ -0,0 +1,22 @@
+// Copyright Epic Games, Inc. All Rights Reserved.
+
+import { DescriptionParser, descriptionParsers } from '../targets'
+
+export class TestCustomDescriptionParser extends DescriptionParser {
+
+	static {
+		descriptionParsers.set("testcustomparser", TestCustomDescriptionParser)
+	}
+
+	override parse(desc: string, lineEnd: string = '\n') {
+		// If the first line is first.last that will be the author tag
+		const authorMatch = desc.match(/^([^ ]*?\.[^ ]*?)\n(.*)/s)
+		if (authorMatch) {
+			this.authorTag = authorMatch[1]
+			this.descFinal = authorMatch[2].split(lineEnd)
+		}
+		else {
+			this.descFinal = desc.split(lineEnd)
+		}
+	}
+}
\ No newline at end of file
diff --git a/Engine/Extras/RoboMerge/v3/src/robo/nodebot.ts b/Engine/Extras/RoboMerge/v3/src/robo/nodebot.ts
--- a/Engine/Extras/RoboMerge/v3/src/robo/nodebot.ts
+++ b/Engine/Extras/RoboMerge/v3/src/robo/nodebot.ts
@@ -21,7 +21,7 @@ import { Context } from './settings';
 import { SlackFile, SlackMessage, SlackMessageField, SlackMessageStyles } from './slack';
 import { PauseState } from './state-interfaces';
 import { newTickJournal, TickJournal } from './tick-journal';
-import { computeTargets, parseDescriptionLines, processOtherBotTargets, getIntegrationOwner, getNodeBotFullName, getNodeBotFullNameForLogging } from './targets';
+import { computeTargets, parseDescription, processOtherBotTargets, getIntegrationOwner, getNodeBotFullName, getNodeBotFullNameForLogging } from './targets';
 import { GraphAPI } from '../new/graph';
 
 /**********************************
@@ -1845,22 +1845,21 @@ export class NodeBot extends PerforceStatefulBot implements NodeBotInterface {
 			}
 		}
 
-		const parse = (lines: string[]) => parseDescriptionLines({
-			lines,
+		const parseDescCommonParams = {
 			isDefaultBot: this.branch.isDefaultBot,
 			graphBotName: this.graphBotName,
 			cl: change.change,
 			aliasesUpper: this.branchGraph.config.aliases.map(s => s.toUpperCase()),
 			macros: combinedMacros,
 			logger: this.nodeBotLogger
-		})
+		}
 
-		const parsedLines = parse(change.desc.split('\n'))
+		const parsedLines = parseDescription(change.desc, {parser: this.branch.changelistParser, ...parseDescCommonParams})
 
 		const commandOverride = (change.commandOverride || '').trim()
 		if (commandOverride) {
 
-			const parsedOverride = parse(commandOverride.split('|').map(s => s.trim()))
+			const parsedOverride = parseDescription(commandOverride, {lineEnd: '|', ...parseDescCommonParams})
 			if (change.accumulateCommandOverride) {
 				parsedLines.append(parsedOverride)
 			}
diff --git a/Engine/Extras/RoboMerge/v3/src/robo/targets.ts b/Engine/Extras/RoboMerge/v3/src/robo/targets.ts
--- a/Engine/Extras/RoboMerge/v3/src/robo/targets.ts
+++ b/Engine/Extras/RoboMerge/v3/src/robo/targets.ts
@@ -264,8 +264,8 @@ export class DescriptionParser {
 		return true
 	}
 
-	parse(lines: string[]) {
-		for (const line of lines) {
+	parse(desc: string, lineEnd?: string) {
+		for (const line of desc.split(lineEnd || '\n').map(s => s.trim())) {
 			const result = tokenizeLine(line)
 			if (!result) {
 				// strip beginning blanks
@@ -335,18 +335,23 @@ export class DescriptionParser {
 	}
 }
 
-export function parseDescriptionLines(args: {
-		lines: string[],
+export const descriptionParsers = new Map<string, typeof DescriptionParser>()
+
+export function parseDescription(desc: string,
+	args: {
 		isDefaultBot: boolean, 
 		graphBotName: string, 
 		cl: number, 
 		aliasesUpper: string[],
 		macros: {[name: string]: string[]},
 		logger: ContextualLogger
-	}
-	) {
-	const lineParser = new DescriptionParser(args.isDefaultBot, args.graphBotName, args.cl, args.aliasesUpper, args.macros)
-	lineParser.parse(args.lines)
+		parser?: string|null,
+		lineEnd?: string
+	}) {
+
+	const descriptionParser = args.parser && descriptionParsers.get(args.parser) || DescriptionParser
+	const lineParser = new descriptionParser(args.isDefaultBot, args.graphBotName, args.cl, args.aliasesUpper, args.macros)
+	lineParser.parse(desc, args.lineEnd)
 	return lineParser
 }
 
@@ -810,7 +815,7 @@ const colors = require('colors')
 colors.enable()
 colors.setTheme(require('colors/themes/generic-logging.js'))
 
-export function runTests(parentLogger: ContextualLogger) {
+export async function runTests(parentLogger: ContextualLogger) {
 	const logger = parentLogger.createChild('Targets')
 
 	const defaultArgs = {
@@ -819,7 +824,8 @@ export function runTests(parentLogger: ContextualLogger) {
 		cl: 1,
 		aliasesUpper: [],
 		macros: {'m': ['#robomerge X, Y'], 'otherbot': ['#robomerge[A] B']},
-		logger
+		logger,
+		parser: null
 	}
 
 	let fails = 0
@@ -855,61 +861,61 @@ export function runTests(parentLogger: ContextualLogger) {
 
 	(() => {
 		testName = 'no commands'
-		const parser = parseDescriptionLines({lines: ['no command'], ...defaultArgs})
+		const parser = parseDescription('no command', defaultArgs)
 		parserAssert(parser, 1, 0, 0, 0, 0, parser.descFinal[0] === 'no command')
 	})();
 
 	(() => {
 		testName = 'multiple targets'
-		const parser = parseDescriptionLines({lines: ['#robomerge A', '#robomerge B'], ...defaultArgs})
+		const parser = parseDescription('#robomerge A\n#robomerge B', defaultArgs)
 		parserAssert(parser, 0, 2, 0, 0, 0, true)
 	})();
 
 	(() => {
 		testName = 'comma separated targets'
-		const parser = parseDescriptionLines({lines: ['#robomerge A, B'], ...defaultArgs})
+		const parser = parseDescription('#robomerge A, B', defaultArgs)
 		parserAssert(parser, 0, 2, 0, 0, 0, true)
 	})();
 
 	(() => {
 		testName = 'space separated targets'
-		const parser = parseDescriptionLines({lines: ['#robomerge A B'], ...defaultArgs})
+		const parser = parseDescription('#robomerge A B', defaultArgs)
 		parserAssert(parser, 0, 2, 0, 0, 0, true)
 	})();
 
 	(() => {
 		testName = 'macro'
-		const parser = parseDescriptionLines({lines: ['#robomerge M'], ...defaultArgs})
+		const parser = parseDescription('#robomerge M', defaultArgs)
 		parserAssert(parser, 0, 2, 0, 1, 0, true)
 	})();
 
 	(() => {
 		testName = 'other bot'
-		const parser = parseDescriptionLines({lines: ['#robomerge[B] A'], ...defaultArgs})
+		const parser = parseDescription('#robomerge[B] A', defaultArgs)
 		parserAssert(parser, 0, 0, 1, 0, 0, true)
 	})();
 
 	(() => {
 		testName = 'all bots'
-		const parser = parseDescriptionLines({lines: ['#robomerge[ALL] A'], ...defaultArgs})
+		const parser = parseDescription('#robomerge[ALL] A', defaultArgs)
 		parserAssert(parser, 0, 1, 0, 0, 0, true)
 	})();
 
 	(() => {
 		testName = 'other bot macro'
-		const parser = parseDescriptionLines({lines: ['#robomerge otherbot'], ...defaultArgs})
+		const parser = parseDescription('#robomerge otherbot', defaultArgs)
 		parserAssert(parser, 0, 0, 1, 1, 0, true)
 	})();
 
 	(() => {
 		testName = 'unknown tag'
-		const parser = parseDescriptionLines({lines: ['#unknown blah'], ...defaultArgs})
+		const parser = parseDescription('#unknown blah', defaultArgs)
 		parserAssert(parser, 1, 0, 0, 0, 0, parser.descFinal[0] === '#unknown blah')
 	})();
 
 	(() => {
 		testName = 'sanitize known tag'
-		const parser = parseDescriptionLines({lines: ['#fyi blah'], ...defaultArgs})
+		const parser = parseDescription('#fyi blah', defaultArgs)
 		assert(parser.descFinal.length === 1 &&
 			parser.descFinal[0].indexOf('#') < 0 && parser.descFinal[0].toLowerCase().indexOf('fyi') >= 0,
 			'known tag sanitized')
@@ -917,12 +923,21 @@ export function runTests(parentLogger: ContextualLogger) {
 
 	(() => {
 		testName = 'sanitize swarm tag'
-		const parser = parseDescriptionLines({lines: ['#review-123 blah'], ...defaultArgs})
+		const parser = parseDescription('#review-123 blah', defaultArgs)
 		assert(parser.descFinal.length === 1 &&
 			parser.descFinal[0].indexOf('#') < 0 && parser.descFinal[0].toLowerCase().indexOf('blah') >= 0,
 			'swarm tag sanitized')
 	})();
 
+	await import('./changelistparsers/customtestparser');
+	(() => {
+		testName = 'custom parser '
+		const parser = parseDescription('roger.rabbit\nRemaining description\nto be parsed', {...defaultArgs, parser: 'testcustomparser'})
+		assert(parser.authorTag === 'roger.rabbit', 'author parsed')
+		assert(parser.descFinal.length === 2, 'description parsed')
+	})();
+
+
 	if (fails === 0) {
 		logger.info(colors.info(`Targets tests succeeded (${assertions} assertions)`))
 	}
