Reduce perforce command spam, particularly while transferring unless in devMode

diff --git a/Engine/Extras/RoboMerge/v3/src/common/perforce.ts b/Engine/Extras/RoboMerge/v3/src/common/perforce.ts
--- a/Engine/Extras/RoboMerge/v3/src/common/perforce.ts
+++ b/Engine/Extras/RoboMerge/v3/src/common/perforce.ts
@@ -298,10 +298,12 @@ export interface IntegrateOpts {
 export interface IntegratedOpts {
 	intoOnly?: boolean
 	startCL?: number
+	quiet?: boolean
 }
 
 export interface SyncParams {
 	opts?: string[]
+	quiet?: boolean
 	edgeServerAddress?: string
 	okToFail?: boolean // default is false
 }
@@ -456,7 +458,7 @@ const MIN_SERVER_VERSION: number = 2020.1
  * This method must succeed before PerforceContext can be used. Otherwise retrieving the Perforce context through
  * getServerContext() will error.
  */
-export async function initializePerforce(logger: ContextualLogger, vault: Vault) {
+export async function initializePerforce(logger: ContextualLogger, vault: Vault, devMode: boolean) {
 	setInterval(() => {
 		const now = Date.now()
 		for (const command of runningPerforceCommands) {
@@ -467,6 +469,7 @@ export async function initializePerforce(logger: ContextualLogger, vault: Vault)
 		}
 	}, 5*60*1000)
 
+	PerforceContext.devMode = devMode
 	const servers = vault.valid && vault.perforceServers
 
 	if (servers) {
@@ -489,6 +492,7 @@ export class PerforceContext {
 	private constructor(private readonly logger: ContextualLogger, private _serverConfig: ServerConfig) {
 	}
 
+	static devMode: boolean
 	static readonly servers: Map<string, ServerConfig> = new Map()
 
 	static async initServer(logger: ContextualLogger, serverID: string, serverDetails?: any) {
@@ -844,7 +848,7 @@ export class PerforceContext {
 	async sync(roboWorkspace: RoboWorkspace, depotPath: string, params?: SyncParams) {
 		const args = ['sync', ...(params && params.opts || []), depotPath]
 		try {
-			return await this.execAndParse(roboWorkspace, args, {serverAddress: params && params.edgeServerAddress})
+			return await this.execAndParse(roboWorkspace, args, {serverAddress: params?.edgeServerAddress, quiet: params?.quiet})
 		}
 		catch (reason) {
 			if (!isExecP4Error(reason)) {
@@ -996,6 +1000,7 @@ export class PerforceContext {
 
 		let failures: string[] = []
 
+		const p4quiet = !PerforceContext.devMode
 		const sourceStreamLength = source.stream!.length
 		const targetFilter = `${target.stream!}/...` == target.path_to ? null : makeRegexFromPerforcePath(target.path_to)
 
@@ -1009,12 +1014,12 @@ export class PerforceContext {
 				case 'add':
 				case 'branch': {
 					// Validate file doesn't exist
-					return targetP4.fstat(roboWorkspace, targetDepotFile)
+					return targetP4.fstat(roboWorkspace, targetDepotFile, p4quiet)
 					.then(fstat => {
 						if (fstat.length == 0 || fstat[0].headAction.includes('delete')) {
 							return targetP4.where(roboWorkspace, targetDepotFile)
 							.then(where => { return this.print(`${file.depotFile}@${source.changelist}`, where[0].path) })
-							.then(() => { return targetP4.add(roboWorkspace, dest_changelist, targetDepotFile, file.type) })
+							.then(() => { return targetP4.add(roboWorkspace, dest_changelist, targetDepotFile, file.type, p4quiet) })
 						}
 						else {
 							failures.push(`File ${targetDepotFile} already exists when being added in CL# ${source.changelist}`)
@@ -1025,11 +1030,11 @@ export class PerforceContext {
 				case 'edit':
 				case 'integrate': {
 					// Validate file exists
-					return targetP4.fstat(roboWorkspace, targetDepotFile)
+					return targetP4.fstat(roboWorkspace, targetDepotFile, p4quiet)
 					.then(fstat => {
 						if (fstat.length > 0 && !fstat[0].headAction.includes('delete')) {
-							return targetP4.sync(roboWorkspace, targetDepotFile, {opts: ['-k']})
-							.then(() => { return targetP4.edit(roboWorkspace, dest_changelist, targetDepotFile, ['-k', '-t', file.type]) })
+							return targetP4.sync(roboWorkspace, targetDepotFile, {opts: ['-k'], quiet: p4quiet})
+							.then(() => { return targetP4.edit(roboWorkspace, dest_changelist, targetDepotFile, ['-k', '-t', file.type], p4quiet) })
 							.then(() => { return this.print(`${file.depotFile}@${source.changelist}`, fstat[0].clientFile) })	
 						}
 						else {
@@ -1039,17 +1044,17 @@ export class PerforceContext {
 					})
 				}
 				case 'delete': {
-					return targetP4.delete(roboWorkspace, dest_changelist, targetDepotFile)
+					return targetP4.delete(roboWorkspace, dest_changelist, targetDepotFile, p4quiet)
 				}
 				case 'move/delete': {
-					return this.integrated(null, file.depotFile, {intoOnly: true, startCL: source.changelist})
+					return this.integrated(null, file.depotFile, {intoOnly: true, startCL: source.changelist, quiet: p4quiet})
 					.then(integrated => {
 						const integ = integrated.find(integ => parseInt(integ.change) == source.changelist)
 						if (integ) {
 							const targetMoveAddDepotFile = target.stream! + integ.fromFile.substring(sourceStreamLength)
-							return targetP4.sync(roboWorkspace, targetDepotFile)
-							.then(() => { return targetP4.edit(roboWorkspace, dest_changelist, targetDepotFile) })
-							.then(() => { return targetP4.move(roboWorkspace, dest_changelist, targetDepotFile, targetMoveAddDepotFile, ['-f']) })
+							return targetP4.sync(roboWorkspace, targetDepotFile, {quiet: p4quiet})
+							.then(() => { return targetP4.edit(roboWorkspace, dest_changelist, targetDepotFile, undefined, p4quiet) })
+							.then(() => { return targetP4.move(roboWorkspace, dest_changelist, targetDepotFile, targetMoveAddDepotFile, ['-f'], p4quiet) })
 							.then(move => {	return this.print(`${integ.fromFile}@${source.changelist}`, move[0].path) })
 						}
 						else {
@@ -1479,8 +1484,8 @@ export class PerforceContext {
 		}
 	}
 
-	async move(roboWorkspace: RoboWorkspace, cl: number, src: string, target: string, opts?: string[]) {
-		return this.execAndParse(roboWorkspace, ['move', ...(opts || []), '-c', cl.toString(), src, target]);
+	async move(roboWorkspace: RoboWorkspace, cl: number, src: string, target: string, opts?: string[], quiet?: boolean) {
+		return this.execAndParse(roboWorkspace, ['move', ...(opts || []), '-c', cl.toString(), src, target], {quiet});
 	}
 
 	/**
@@ -1552,36 +1557,36 @@ export class PerforceContext {
 	}
 
 	/** Check out a file into a specific changelist ** ASYNC ** */
-	edit(roboWorkspace: RoboWorkspace, cl: number, filePath: string, additionalArgs?: string[]) {
+	edit(roboWorkspace: RoboWorkspace, cl: number, filePath: string, additionalArgs?: string[], quiet?: boolean) {
 		const args = [
 			'edit', 
 			'-c', cl.toString(), 
 			...(additionalArgs || []),
 			filePath
 		]
-		return this._execP4(roboWorkspace, args);
+		return this._execP4(roboWorkspace, args, {quiet});
 	}
 
 	/** Add a file into a specific changelist ** ASYNC ** */
-	add(roboWorkspace: RoboWorkspace, cl: number, filePath: string, filetype?: string) {
+	add(roboWorkspace: RoboWorkspace, cl: number, filePath: string, filetype?: string, quiet?: boolean) {
 		const args = [
 			'add', 
 			'-c', cl.toString(), 
 			...(filetype ? ['-t', filetype] : []),
 			filePath
 		]
-		return this._execP4(roboWorkspace, args);
+		return this._execP4(roboWorkspace, args, {quiet});
 	}
 
 	/** Mark a file for delete into a specific changelist ** ASYNC ** */
-	delete(roboWorkspace: RoboWorkspace, cl: number, filePath: string) {
+	delete(roboWorkspace: RoboWorkspace, cl: number, filePath: string, quiet?: boolean) {
 		const args = [
 			'delete', 
 			'-c', cl.toString(),
 			'-v',
 			filePath
 		]
-		return this._execP4(roboWorkspace, args);
+		return this._execP4(roboWorkspace, args, {quiet});
 	}
 
 	async describe(cl: number, maxFiles?: number, includeShelved?: boolean) {
@@ -1665,7 +1670,7 @@ export class PerforceContext {
 	}
 
 	where(roboWorkspace: RoboWorkspace, clientPath: string) {
-		return this.execAndParse(roboWorkspace, ['where', clientPath])
+		return this.execAndParse(roboWorkspace, ['where', clientPath], {quiet: !PerforceContext.devMode})
 	}
 
 	async filelog(roboWorkspace: RoboWorkspace, depotPath: string, beginRev: string, endRev: string, longOutput = false) {
@@ -1693,9 +1698,9 @@ export class PerforceContext {
 		}
 	}
 
-	async fstat(roboWorkspace: RoboWorkspace, depotPath: string) {
+	async fstat(roboWorkspace: RoboWorkspace, depotPath: string, quiet?: boolean) {
 		try {
-			return await this.execAndParse(roboWorkspace, ['fstat', depotPath])
+			return await this.execAndParse(roboWorkspace, ['fstat', depotPath], {quiet})
 		}
 		catch (reason) {
 			if (!isExecP4Error(reason)) {
@@ -1725,7 +1730,7 @@ export class PerforceContext {
 		args.push(depotPath)
 
 		try {
-			return await this.execAndParse(roboWorkspace, args)
+			return await this.execAndParse(roboWorkspace, args, {quiet: opts?.quiet})
 		}
 		catch (reason) {
 			if (!isExecP4Error(reason)) {
diff --git a/Engine/Extras/RoboMerge/v3/src/robo/robo.ts b/Engine/Extras/RoboMerge/v3/src/robo/robo.ts
--- a/Engine/Extras/RoboMerge/v3/src/robo/robo.ts
+++ b/Engine/Extras/RoboMerge/v3/src/robo/robo.ts
@@ -554,7 +554,7 @@ async function _initGraphBots(logger: ContextualLogger) {
 async function main(logger: ContextualLogger) {
 	while (true) {
 		try {
-			await initializePerforce(logger, vault);
+			await initializePerforce(logger, vault, args.devMode);
 			break;
 		}
 		catch (err) {
