Stringify failed user IDs

diff --git a/Engine/Extras/RoboMerge/v3/src/robo/notifications.ts b/Engine/Extras/RoboMerge/v3/src/robo/notifications.ts
--- a/Engine/Extras/RoboMerge/v3/src/robo/notifications.ts
+++ b/Engine/Extras/RoboMerge/v3/src/robo/notifications.ts
@@ -401,7 +401,7 @@ export class SlackMessages {
 							/* this is expected and fine */ 
 						}
 						else {
-							const errorMsg = `Error inviting ${emailAddress} (${user}) to channel <#${channel}>: ${result.error}\n${result.failed_user_ids}`
+							const errorMsg = `Error inviting ${emailAddress} (${user}) to channel <#${channel}>: ${result.error}\n${JSON.stringify(result.failed_user_ids)}`
 							this.smLogger.error(errorMsg)
 							postToRobomergeAlerts(errorMsg)
 						}
