Provide way to use DescriptionParser to provide the true author to use from the Change

diff --git a/Engine/Extras/RoboMerge/v3/src/robo/nodebot.ts b/Engine/Extras/RoboMerge/v3/src/robo/nodebot.ts
--- a/Engine/Extras/RoboMerge/v3/src/robo/nodebot.ts
+++ b/Engine/Extras/RoboMerge/v3/src/robo/nodebot.ts
@@ -21,7 +21,7 @@ import { Context } from './settings';
 import { SlackFile, SlackMessage, SlackMessageField, SlackMessageStyles } from './slack';
 import { PauseState } from './state-interfaces';
 import { newTickJournal, TickJournal } from './tick-journal';
-import { computeTargets, parseDescription, processOtherBotTargets, getIntegrationOwner, getNodeBotFullName, getNodeBotFullNameForLogging } from './targets';
+import { computeTargets, descriptionParsers, parseDescription, processOtherBotTargets, getIntegrationOwner, getNodeBotFullName, getNodeBotFullNameForLogging } from './targets';
 import { GraphAPI } from '../new/graph';
 
 /**********************************
@@ -1502,6 +1502,29 @@ export class NodeBot extends PerforceStatefulBot implements NodeBotInterface {
 		}
 	}
 
+	private getChangeAuthor(change: Change) {
+		const descriptionParser = this.branch.changelistParser && descriptionParsers.get(this.branch.changelistParser)
+		if (descriptionParser) {
+			const actualAuthor = descriptionParser.getAuthor(change)
+			if (actualAuthor) {
+				return actualAuthor
+			}
+		}
+
+		// If the user of the CL is robomerge, lookup the actual author from the tag
+		if (change.user === "robomerge") {
+			for (const line of change.desc.split('\n')) {
+				const authorMatch = line.match("#ROBOMERGE-AUTHOR: (.*)")
+				if (authorMatch) {
+					return authorMatch[1]
+				}
+			}
+			this.nodeBotLogger.warn(`Processing change ${change.change} by user robomerge. Now: ${Date.now()} Change time: ${change.time}`)
+		}
+
+		return change.user
+	}
+
 	private async _processListOfChanges(availableEdges: Map<string, EdgeBot>, allChanges: Change[]) {
 		// list of changes except reconsiders
 		const changes: Change[] = []
@@ -1533,24 +1556,11 @@ export class NodeBot extends PerforceStatefulBot implements NodeBotInterface {
 		for (let changeIndex = 0; changeIndex < changes.length; ++changeIndex) {
 			const change = changes[changeIndex]
 
-			// If the user of the CL is robomerge, lookup the actual author from the tag
-			if (change.user === "robomerge") {
-				let bAuthorFound = false
-				for (const line of change.desc.split('\n')) {
-					const authorMatch = line.match("#ROBOMERGE-AUTHOR: (.*)")
-					if (authorMatch) {
-						change.user = authorMatch[1]
-						bAuthorFound = true
-						break
-					}
-				}
-				if (!bAuthorFound) {
-					this.nodeBotLogger.warn(`Processing change ${change.change} by user robomerge. Now: ${Date.now()} Change time: ${change.time}`)
-				}
-			}
+			change.user = this.getChangeAuthor(change)
+			
 			// If this is a change submitted via swarm on behalf of someone don't process the
 			// change for a few minutes to give the change owner command time to be processed
-			else if (change.client.startsWith("swarm-")) {
+			if (change.client.startsWith("swarm-")) {
 				// TODO: Use the Swarm API to lookup the actual author
 				if ((Date.now() - ((change.time || 0) * 1000)) < 120000) {
 					this.nodeBotLogger.info(`Delaying processing of ${change.change} to give swarm time to change owner.  Now: ${Date.now()} Change time: ${change.time}`)
diff --git a/Engine/Extras/RoboMerge/v3/src/robo/targets.ts b/Engine/Extras/RoboMerge/v3/src/robo/targets.ts
--- a/Engine/Extras/RoboMerge/v3/src/robo/targets.ts
+++ b/Engine/Extras/RoboMerge/v3/src/robo/targets.ts
@@ -4,6 +4,7 @@ import { Branch, MergeAction, MergeMode } from './branch-interfaces';
 import { BranchGraphInterface, PendingChange, Target, TargetInfo } from './branch-interfaces';
 import { BotName, Edge, GraphAPI, Node, makeTargetName } from '../new/graph';
 import { ContextualLogger } from '../common/logger';
+import { Change } from '../common/perforce';
 
 // const NEW_STUFF = false
 
@@ -333,6 +334,10 @@ export class DescriptionParser {
 			}
 		}
 	}
+
+	static getAuthor(_: Change) : string|undefined {
+		return undefined
+	}
 }
 
 export const descriptionParsers = new Map<string, typeof DescriptionParser>()
