Fix failure to clean up shelved changelists

diff --git a/Engine/Extras/RoboMerge/v3/src/common/p4util.ts b/Engine/Extras/RoboMerge/v3/src/common/p4util.ts
--- a/Engine/Extras/RoboMerge/v3/src/common/p4util.ts
+++ b/Engine/Extras/RoboMerge/v3/src/common/p4util.ts
@@ -47,7 +47,7 @@ export async function cleanWorkspace(logger: ContextualLogger, p4: PerforceConte
 	// revert any pending changes left over from previous runs
 	for (const change of changes) {
 		const changeStr = `CL ${change.change}: ${change.desc}`
-		if (change.shelved) {
+		if (change.shelved !== undefined) {
 			logger.info(`Attempting to delete shelved files in ${changeStr}`)
 			try {
 				await p4.delete_shelved(workspace, change.change)
