Add getUser function and rework getEmail to work in terms of it

diff --git a/Engine/Extras/RoboMerge/v3/src/common/perforce.ts b/Engine/Extras/RoboMerge/v3/src/common/perforce.ts
--- a/Engine/Extras/RoboMerge/v3/src/common/perforce.ts
+++ b/Engine/Extras/RoboMerge/v3/src/common/perforce.ts
@@ -1545,15 +1545,30 @@ export class PerforceContext {
 		return this._execP4(roboWorkspace, ["shelve", "-d", "-c", changelist.toString(), "-f"]);
 	}
 
-	// get the email (according to P4) for a specific user
-	async getEmail(username: string) {
-		let m = null
+	async getUser(username: string) {
 		if (username.length > 0 && !username.startsWith('@') && !username.startsWith('<') && !isUserAKnownBot(username)) {
-			const output = await this._execP4(null, ['user', '-o', username]);
-			// look for the email field
-			m = output.match(/\nEmail:\s+([^\n]+)\n/);
+			try {
+				return (await this.execAndParse(null, ['users', username], {quiet:true}))[0]
+			}
+			catch (reason) {
+				if (!isExecP4Error(reason)) {
+					throw reason
+				}
+	
+				let [err, output] = reason
+				if (output.includes('no such user')) {
+					return null
+				}
+				throw err
+			}
 		}
-		return m && m[1];
+		return null
+	}
+
+	// get the email (according to P4) for a specific user
+	async getEmail(username: string) {
+		const user = await this.getUser(username)
+		return user && user.Email
 	}
 
 	/** Check out a file into a specific changelist ** ASYNC ** */
