Make getAuthor async

diff --git a/Engine/Extras/RoboMerge/v3/src/robo/nodebot.ts b/Engine/Extras/RoboMerge/v3/src/robo/nodebot.ts
--- a/Engine/Extras/RoboMerge/v3/src/robo/nodebot.ts
+++ b/Engine/Extras/RoboMerge/v3/src/robo/nodebot.ts
@@ -1502,10 +1502,10 @@ export class NodeBot extends PerforceStatefulBot implements NodeBotInterface {
 		}
 	}
 
-	private getChangeAuthor(change: Change) {
+	private async getChangeAuthor(change: Change) {
 		const descriptionParser = this.branch.changelistParser && descriptionParsers.get(this.branch.changelistParser)
 		if (descriptionParser) {
-			const actualAuthor = descriptionParser.getAuthor(change)
+			const actualAuthor = await descriptionParser.getAuthor(change, this.nodeBotLogger)
 			if (actualAuthor) {
 				return actualAuthor
 			}
@@ -1556,7 +1556,7 @@ export class NodeBot extends PerforceStatefulBot implements NodeBotInterface {
 		for (let changeIndex = 0; changeIndex < changes.length; ++changeIndex) {
 			const change = changes[changeIndex]
 
-			change.user = this.getChangeAuthor(change)
+			change.user = await this.getChangeAuthor(change)
 			
 			// If this is a change submitted via swarm on behalf of someone don't process the
 			// change for a few minutes to give the change owner command time to be processed
diff --git a/Engine/Extras/RoboMerge/v3/src/robo/targets.ts b/Engine/Extras/RoboMerge/v3/src/robo/targets.ts
--- a/Engine/Extras/RoboMerge/v3/src/robo/targets.ts
+++ b/Engine/Extras/RoboMerge/v3/src/robo/targets.ts
@@ -335,7 +335,7 @@ export class DescriptionParser {
 		}
 	}
 
-	static getAuthor(_: Change) : string|undefined {
+	static async getAuthor(_change: Change, _logger: ContextualLogger) : Promise<string|undefined> {
 		return undefined
 	}
 }
