Fix possible owner not being displayed if using the <!subteam syntax for resolver

diff --git a/Engine/Extras/RoboMerge/v3/public/js/boilerplate.js b/Engine/Extras/RoboMerge/v3/public/js/boilerplate.js
--- a/Engine/Extras/RoboMerge/v3/public/js/boilerplate.js
+++ b/Engine/Extras/RoboMerge/v3/public/js/boilerplate.js
@@ -1009,7 +1009,12 @@ function renderNameCell_Common(data, botname) {
 			if (data.blockage.acknowledger) {
 				msg += `Acknowledged by ${data.blockage.acknowledger}, please contact for more information.<br />`;
 			} else {
-				msg += `Blocked but not acknowledged. Possible owner: ${data.blockage.owner}<br />`;
+				owner = data.blockage.owner
+				ownerMatch = owner?.match(/<!.*\|(.*)>/)
+				if (ownerMatch) {
+					owner = `@${ownerMatch[1]}`
+				}
+				msg += `Blocked but not acknowledged. Possible owner: ${owner}<br />`;
 			}
 		}
 
