Improve approvals notifications, particularly for blocked approvals.

diff --git a/Engine/Extras/RoboMerge/v3/src/robo/nodebot.ts b/Engine/Extras/RoboMerge/v3/src/robo/nodebot.ts
--- a/Engine/Extras/RoboMerge/v3/src/robo/nodebot.ts
+++ b/Engine/Extras/RoboMerge/v3/src/robo/nodebot.ts
@@ -2251,7 +2251,7 @@ export class NodeBot extends PerforceStatefulBot implements NodeBotInterface {
 			const slackOwner = ownerEmail ?  await this.slackMessages.getSlackUser(ownerEmail) : null
 
 			const fields: SlackMessageField[] = [
-				{title: 'Change', short: true, value: makeClLink(pending.newCl)}
+				{title: 'Shelved Change', short: true, value: makeClLink(pending.newCl)}
 			]
 
 			let message = (slackAuthor ? `<@${slackAuthor}>` : `@${pending.change.author}`) +
diff --git a/Engine/Extras/RoboMerge/v3/src/robo/notifications.ts b/Engine/Extras/RoboMerge/v3/src/robo/notifications.ts
--- a/Engine/Extras/RoboMerge/v3/src/robo/notifications.ts
+++ b/Engine/Extras/RoboMerge/v3/src/robo/notifications.ts
@@ -515,8 +515,6 @@ export class BotNotifications implements BotEventHandler {
 			return
 		}
 
-		const cl = blockage.approval ? blockage.approval.shelfCl : changeInfo.cl
-
 		// or integration failure (better wording? exclusive check-out?)
 		const sourceBranch = changeInfo.branch
 		let targetBranch
@@ -535,20 +533,24 @@ export class BotNotifications implements BotEventHandler {
 
 		const isBotUser = isUserAKnownBot(blockage.owner)
 		const text =
-			blockage.approval ?				`${channelPing}'s change needs to be approved in <#${blockage.approval.settings.channelId}>` :
+			blockage.approval ?				`${channelPing}'s change needs to be approved\n\n${blockage.approval.settings.description}` :
 			isBotUser ? 									`Blockage caused by \`${blockage.owner}\` commit!` :
 															`${channelPing}, please resolve the following ${issue}:`
 
-		const message = this.makeSlackChannelMessage(
+		let message = this.makeSlackChannelMessage(
 			`${sourceBranch.name} blocked! (${issue})`,
 			text,
 			SlackMessageStyles.DANGER, 
-			makeClLink(cl), 
+			makeClLink(changeInfo.cl), 
 			sourceBranch, 
 			targetBranch, 
 			changeInfo.author
 		);
 
+		if (blockage.approval) {
+			message.fields?.push({title: 'Shelved Change', short: true, value: makeClLink(blockage.approval.shelfCl)})
+		}
+
 		let messagesToPost = []
 		let channelsToPostTo = []
 		let usersToInvite: Set<string> = new Set()
@@ -653,7 +655,7 @@ export class BotNotifications implements BotEventHandler {
 			if (blockage.approval) {
 				dm = {
 					title: 'Approval needed to commit to ' + targetBranch.name,
-					text: `Your change has been shelved in ${makeClLink(cl)} and sent to <#${blockage.approval.settings.channelId}> for approval\n\n` +
+					text: `Your change has been shelved in ${makeClLink(blockage.approval.shelfCl)} and sent to <#${blockage.approval.settings.channelId}> for approval\n\n` +
 							blockage.approval.settings.description,
 					channel: "",
 					mrkdwn: true
@@ -672,7 +674,7 @@ export class BotNotifications implements BotEventHandler {
 					throw error
 				}
 
-				dm = this.makeSlackDirectMessage(dmText, changeInfo.cl, cl, targetBranch.name, urls)
+				dm = this.makeSlackDirectMessage(dmText, changeInfo.cl, changeInfo.cl, targetBranch.name, urls)
 			}
 
 			this.slackMessages.postDM(userEmail, changeInfo.cl, targetBranch, dm)
@@ -790,7 +792,7 @@ export class BotNotifications implements BotEventHandler {
 					"0",
 					false
 				)
-		)
+			)
 
 			this.slackMessages.postDM([REDACTED_EMAIL]`, 0, "TARGETBRANCH", messageOpts, false)
 		}
