Modify transfer logic to avoid overloading server with commands for large changelists

diff --git a/Engine/Extras/RoboMerge/v3/src/common/perforce.ts b/Engine/Extras/RoboMerge/v3/src/common/perforce.ts
--- a/Engine/Extras/RoboMerge/v3/src/common/perforce.ts
+++ b/Engine/Extras/RoboMerge/v3/src/common/perforce.ts
@@ -845,8 +845,11 @@ export class PerforceContext {
 	}
 
 	// sync the depot path specified
-	async sync(roboWorkspace: RoboWorkspace, depotPath: string, params?: SyncParams) {
-		const args = ['sync', ...(params && params.opts || []), depotPath]
+	async sync(roboWorkspace: RoboWorkspace, depotPath: string|string[], params?: SyncParams) {
+		if (typeof depotPath === "string") {
+			depotPath = [depotPath]
+		}
+		const args = ['sync', ...(params && params.opts || []), ...depotPath]
 		try {
 			return await this.execAndParse(roboWorkspace, args, {serverAddress: params?.edgeServerAddress, quiet: params?.quiet})
 		}
@@ -863,7 +866,7 @@ export class PerforceContext {
 			}
 		}
 
-		return null
+		return []
 	}
 
 	newWorkspace(workspaceName: string, params: any, edgeServer?: EdgeServer) {
@@ -1004,7 +1007,15 @@ export class PerforceContext {
 		const sourceStreamLength = source.stream!.length
 		const targetFilter = `${target.stream!}/...` == target.path_to ? null : makeRegexFromPerforcePath(target.path_to)
 
-		await Promise.all(files.map(file => {
+		let sortedFiles = new Map<string, any[]>()
+		const addSortedFile = (key: string, file: any) => {
+			if (!sortedFiles.has(key)) {
+				sortedFiles.set(key,[])
+			}
+			sortedFiles.get(key)!.push(file)
+		}
+
+		files.forEach(file => {
 			const targetDepotFile = target.stream! + file.depotFile.substring(sourceStreamLength)
 			if (targetFilter && !targetDepotFile.match(targetFilter)) {
 				// filter out files that don't match the target.path_to
@@ -1013,40 +1024,69 @@ export class PerforceContext {
 			switch(file.action) {
 				case 'add':
 				case 'branch': {
-					// Validate file doesn't exist
-					return targetP4.fstat(roboWorkspace, targetDepotFile, p4quiet)
-					.then(fstat => {
-						if (fstat.length == 0 || fstat[0].headAction.includes('delete')) {
-							return targetP4.where(roboWorkspace, targetDepotFile)
-							.then(where => { return this.print(`${file.depotFile}@${source.changelist}`, where[0].path) })
-							.then(() => { return targetP4.add(roboWorkspace, dest_changelist, targetDepotFile, file.type, p4quiet) })
-						}
-						else {
-							failures.push(`File ${targetDepotFile} already exists when being added in CL# ${source.changelist}`)
-							return
-						}
-					})
+					addSortedFile(`add:${file.type}`, [file, targetDepotFile])
+					break
 				}
 				case 'edit':
 				case 'integrate': {
-					// Validate file exists
-					return targetP4.fstat(roboWorkspace, targetDepotFile, p4quiet)
-					.then(fstat => {
-						if (fstat.length > 0 && !fstat[0].headAction.includes('delete')) {
-							return targetP4.sync(roboWorkspace, targetDepotFile, {opts: ['-k'], quiet: p4quiet})
-							.then(() => { return targetP4.edit(roboWorkspace, dest_changelist, targetDepotFile, ['-k', '-t', file.type], p4quiet) })
-							.then(() => { return this.print(`${file.depotFile}@${source.changelist}`, fstat[0].clientFile) })	
-						}
-						else {
-							failures.push(`File ${targetDepotFile} does not exist when being edited in CL# ${source.changelist}`)
-							return
-						}
-					})
+					addSortedFile(`edit:${file.type}`, [file, targetDepotFile])
+					break
 				}
 				case 'delete': {
-					return targetP4.delete(roboWorkspace, dest_changelist, targetDepotFile, p4quiet)
+					addSortedFile('delete', [file, targetDepotFile])
+					break
 				}
 				case 'move/delete': {
+					addSortedFile('move', [file, targetDepotFile])
+					break
+				}
+				case 'move/add': break // move/add is handled from the move/delete portion of the pair
+
+				default: {
+					failures.push(`Unhandled transfer action ${file.action} for ${file.depotFile} in CL# ${source.changelist}`)
+				}
+			}
+		})
+
+		let promises: Promise<any>[] = []
+		let printCommands: any[] = []
+
+		const FILES_PER_ACTION = 100
+
+		sortedFiles.forEach((files, action) => {
+			if (action.startsWith('add')) {
+				const fileType = action.slice(4)
+				for (let i=0; i < files.length; i += FILES_PER_ACTION) {
+					const filesSlice = files.slice(i,i+FILES_PER_ACTION)
+					const targetDepotFiles = filesSlice.map(([_, targetDepotFile]) => targetDepotFile)
+					promises.push(
+						targetP4.add(roboWorkspace, dest_changelist, targetDepotFiles, fileType, p4quiet)
+						        .then(adds => adds.forEach((add, index) => printCommands.push([`${filesSlice[index][0].depotFile}@${source.changelist}`, add.clientFile])))
+					)
+				}
+			}
+			else if (action.startsWith('edit')) {
+				const fileType = action.slice(5)
+				for (let i=0; i < files.length; i += FILES_PER_ACTION) {
+					const filesSlice = files.slice(i,i+FILES_PER_ACTION)
+					const targetDepotFiles = filesSlice.map(([_, targetDepotFile]) => targetDepotFile)
+					promises.push(
+						targetP4.sync(roboWorkspace, targetDepotFiles, {opts: ['-k'], quiet: p4quiet})
+						        .then(syncs => { 
+									syncs.forEach((sync, index) => printCommands.push([`${filesSlice[index][0].depotFile}@${source.changelist}`, sync.clientFile]))
+									return targetP4.edit(roboWorkspace, dest_changelist, targetDepotFiles, ['-k', '-t', fileType], p4quiet)
+								}
+						)
+					)
+				}
+			}
+			else if (action === 'delete') {
+				for (let i=0; i < files.length; i += FILES_PER_ACTION) {
+					promises.push(targetP4.delete(roboWorkspace, dest_changelist, files.slice(i,i+FILES_PER_ACTION).map(([_, targetDepotFile]) => targetDepotFile), p4quiet))
+				}
+			}
+			else if (action === 'move') {
+				promises.push(...files.map(([file, targetDepotFile]) => {
 					return this.integrated(null, file.depotFile, {intoOnly: true, startCL: source.changelist, quiet: p4quiet})
 					.then(integrated => {
 						const integ = integrated.find(integ => parseInt(integ.change) == source.changelist)
@@ -1055,22 +1095,29 @@ export class PerforceContext {
 							return targetP4.sync(roboWorkspace, targetDepotFile, {quiet: p4quiet})
 							.then(() => { return targetP4.edit(roboWorkspace, dest_changelist, targetDepotFile, undefined, p4quiet) })
 							.then(() => { return targetP4.move(roboWorkspace, dest_changelist, targetDepotFile, targetMoveAddDepotFile, ['-f'], p4quiet) })
-							.then(move => {	return this.print(`${integ.fromFile}@${source.changelist}`, move[0].path) })
+							.then(move => printCommands.push([`${integ.fromFile}@${source.changelist}`, move[0].path]))
 						}
 						else {
 							failures.push(`Unable to look up integrated record for ${file.depotFile} in CL# ${source.changelist}`)
 							return
 						}
 					})
-				}
-				case 'move/add': return // move/add is handled from the move/delete portion of the pair
+				}))
+			}
+			else {
+				throw new Error(`Unexpected action ${action} in sortedFiles map`)
+			}
+		})
 
-				default: {
-					failures.push(`Unhandled transfer action ${file.action} for ${file.depotFile} in CL# ${source.changelist}`)
-					return
-				}
+		await Promise.all(promises)
+
+		const PARALLEL_PRINTS = 25
+		const doPrints = async (startIndex: number) => {
+			for (let i=startIndex; i < printCommands.length; i += PARALLEL_PRINTS) {
+				await this.print(printCommands[i][0],printCommands[i][1])
 			}
-		}))
+		}
+		await Promise.all(Array.from(Array(PARALLEL_PRINTS).keys()).map(n => doPrints(n)))
 
 		return failures
 	}
@@ -1572,36 +1619,45 @@ export class PerforceContext {
 	}
 
 	/** Check out a file into a specific changelist ** ASYNC ** */
-	edit(roboWorkspace: RoboWorkspace, cl: number, filePath: string, additionalArgs?: string[], quiet?: boolean) {
+	edit(roboWorkspace: RoboWorkspace, cl: number, filePath: string|string[], additionalArgs?: string[], quiet?: boolean) {
+		if (typeof filePath === "string") {
+			filePath = [filePath]
+		}
 		const args = [
 			'edit', 
 			'-c', cl.toString(), 
 			...(additionalArgs || []),
-			filePath
+			...filePath
 		]
-		return this._execP4(roboWorkspace, args, {quiet});
+		return this.execAndParse(roboWorkspace, args, {quiet});
 	}
 
 	/** Add a file into a specific changelist ** ASYNC ** */
-	add(roboWorkspace: RoboWorkspace, cl: number, filePath: string, filetype?: string, quiet?: boolean) {
+	add(roboWorkspace: RoboWorkspace, cl: number, filePath: string|string[], filetype?: string, quiet?: boolean) {
+		if (typeof filePath === "string") {
+			filePath = [filePath]
+		}
 		const args = [
 			'add', 
 			'-c', cl.toString(), 
 			...(filetype ? ['-t', filetype] : []),
-			filePath
+			...filePath
 		]
-		return this._execP4(roboWorkspace, args, {quiet});
+		return this.execAndParse(roboWorkspace, args, {quiet});
 	}
 
 	/** Mark a file for delete into a specific changelist ** ASYNC ** */
-	delete(roboWorkspace: RoboWorkspace, cl: number, filePath: string, quiet?: boolean) {
+	delete(roboWorkspace: RoboWorkspace, cl: number, filePath: string|string[], quiet?: boolean) {
+		if (typeof filePath === "string") {
+			filePath = [filePath]
+		}
 		const args = [
 			'delete', 
 			'-c', cl.toString(),
 			'-v',
-			filePath
+			...filePath
 		]
-		return this._execP4(roboWorkspace, args, {quiet});
+		return this.execAndParse(roboWorkspace, args, {quiet});
 	}
 
 	async describe(cl: number, maxFiles?: number, includeShelved?: boolean) {
@@ -1684,8 +1740,11 @@ export class PerforceContext {
 		await this.editChange(roboWorkspace, changelist, opts)
 	}
 
-	where(roboWorkspace: RoboWorkspace, clientPath: string) {
-		return this.execAndParse(roboWorkspace, ['where', clientPath], {quiet: !PerforceContext.devMode})
+	where(roboWorkspace: RoboWorkspace, clientPath: string|string[]) {
+		if (typeof clientPath === "string") {
+			clientPath = [clientPath]
+		}
+		return this.execAndParse(roboWorkspace, ['where', ...clientPath], {quiet: !PerforceContext.devMode})
 	}
 
 	async filelog(roboWorkspace: RoboWorkspace, depotPath: string, beginRev: string, endRev: string, longOutput = false) {
