Fix shelves for approvals being deleted when doing blocking approvals

diff --git a/Engine/Extras/RoboMerge/v3/src/robo/edgebot.ts b/Engine/Extras/RoboMerge/v3/src/robo/edgebot.ts
--- a/Engine/Extras/RoboMerge/v3/src/robo/edgebot.ts
+++ b/Engine/Extras/RoboMerge/v3/src/robo/edgebot.ts
@@ -926,7 +926,7 @@ class EdgeBotImpl extends PerforceStatefulBot {
 
 		// abort shelve if this is a buildmachine / robomerge change (unless we are forcing the shelf)
 		let failed = false
-		if (isUserAKnownBot(owner) && !pending.change.forceCreateAShelf) {
+		if (isUserAKnownBot(owner) && !pending.change.forceCreateAShelf && !forApproval) {
 			// revert the changes locally
 			this._log_action(`Reverting shelf due to '${owner}' being a known bot`)
 			await this.p4.revert(destRoboWorkspace, changenum, undefined, edgeServerAddress)
@@ -942,7 +942,7 @@ class EdgeBotImpl extends PerforceStatefulBot {
 			}
 		}
 
-		if (failed) {
+		if (failed && !forApproval) {
 			// abort abort!
 			// delete the cl
 			this._log_action(`Deleting CL ${changenum}`)
