Catch transfer errors and report them appropriately
Filter file already exists errors on add as errors

diff --git a/Engine/Extras/RoboMerge/v3/src/common/perforce.ts b/Engine/Extras/RoboMerge/v3/src/common/perforce.ts
--- a/Engine/Extras/RoboMerge/v3/src/common/perforce.ts
+++ b/Engine/Extras/RoboMerge/v3/src/common/perforce.ts
@@ -1062,6 +1062,7 @@ export class PerforceContext {
 					promises.push(
 						targetP4.add(roboWorkspace, dest_changelist, targetDepotFiles, fileType, p4quiet)
 						        .then(adds => adds.forEach((add, index) => printCommands.push([`${filesSlice[index][0].depotFile}@${source.changelist}`, add.clientFile])))
+						        .catch(reason => failures.push(reason))
 					)
 				}
 			}
@@ -1075,14 +1076,18 @@ export class PerforceContext {
 						        .then(syncs => { 
 									syncs.forEach((sync, index) => printCommands.push([`${filesSlice[index][0].depotFile}@${source.changelist}`, sync.clientFile]))
 									return targetP4.edit(roboWorkspace, dest_changelist, targetDepotFiles, ['-k', '-t', fileType], p4quiet)
-								}
-						)
+									               .catch(reason => failures.push(reason))
+						        })
+						        .catch(reason => failures.push(reason))
 					)
 				}
 			}
 			else if (action === 'delete') {
 				for (let i=0; i < files.length; i += FILES_PER_ACTION) {
-					promises.push(targetP4.delete(roboWorkspace, dest_changelist, files.slice(i,i+FILES_PER_ACTION).map(([_, targetDepotFile]) => targetDepotFile), p4quiet))
+					promises.push(
+						targetP4.delete(roboWorkspace, dest_changelist, files.slice(i,i+FILES_PER_ACTION).map(([_, targetDepotFile]) => targetDepotFile), p4quiet)
+						.catch(reason => failures.push(reason))
+					)
 				}
 			}
 			else if (action === 'move') {
@@ -1096,6 +1101,7 @@ export class PerforceContext {
 							.then(() => { return targetP4.edit(roboWorkspace, dest_changelist, targetDepotFile, undefined, p4quiet) })
 							.then(() => { return targetP4.move(roboWorkspace, dest_changelist, targetDepotFile, targetMoveAddDepotFile, ['-f'], p4quiet) })
 							.then(move => printCommands.push([`${integ.fromFile}@${source.changelist}`, move[0].path]))
+							.catch(reason => failures.push(reason))
 						}
 						else {
 							failures.push(`Unable to look up integrated record for ${file.depotFile} in CL# ${source.changelist}`)
@@ -1111,13 +1117,15 @@ export class PerforceContext {
 
 		await Promise.all(promises)
 
-		const PARALLEL_PRINTS = 25
-		const doPrints = async (startIndex: number) => {
-			for (let i=startIndex; i < printCommands.length; i += PARALLEL_PRINTS) {
-				await this.print(printCommands[i][0],printCommands[i][1])
+		if (failures.length == 0) {
+			const PARALLEL_PRINTS = 25
+			const doPrints = async (startIndex: number) => {
+				for (let i=startIndex; i < printCommands.length; i += PARALLEL_PRINTS) {
+					await this.print(printCommands[i][0],printCommands[i][1])
+				}
 			}
+			await Promise.all(Array.from(Array(PARALLEL_PRINTS).keys()).map(n => doPrints(n)))
 		}
-		await Promise.all(Array.from(Array(PARALLEL_PRINTS).keys()).map(n => doPrints(n)))
 
 		return failures
 	}
@@ -2092,7 +2100,7 @@ export class PerforceContext {
 		}
 
 		let error = result
-						.filter((r) => Object.hasOwn(r,'data') && Object.hasOwn(r,'generic') && Object.hasOwn(r,'severity'))
+						.filter((r) => Object.hasOwn(r,'data') && (Object.hasOwn(r,'generic') && Object.hasOwn(r,'severity')) || (Object.hasOwn(r,'level')))
 						.map((r) => r.data)
 						.join('')
 
