Improve direct message details
Fix some Robomerge divergences
#rnx

diff --git a/Engine/Extras/RoboMerge/v3/run-tests.sh b/Engine/Extras/RoboMerge/v3/run-tests.sh
--- a/Engine/Extras/RoboMerge/v3/run-tests.sh
+++ b/Engine/Extras/RoboMerge/v3/run-tests.sh
@@ -1,5 +1,3 @@
-set -e
-
 RED='\033[0;31m'
 GREEN='\033[0;32m'
 NO_COLOUR='\033[0m'
@@ -34,17 +32,29 @@ cleanup() {
 
 cleanup
 
-docker network create --driver bridge robomerge_functtest_network
-
+docker network create --driver bridge robomerge_functtest_network &
 
 docker build -t helix -f helix.Dockerfile .
 docker build -t helix-and-node -f helix-and-node.Dockerfile .
 
-docker build -t p4docker -f Dockerfile.p4docker .
-docker build -t robomerge -f Dockerfile .
+docker build -t p4docker -f Dockerfile.p4docker . &
+docker build -t robomerge -f Dockerfile . &
+robo_build_pid=$!
+
+docker build -t robomerge_functionaltests -f functional_tests/tstests.Dockerfile . &
+tests_build_pid=$!
+
+wait $robo_build_pid
+if [ $? -ne 0 ]; then
+    exit 1
+fi
 
-docker build -t robomerge_functionaltests -f functional_tests/tstests.Dockerfile .
+wait $tests_build_pid
+if [ $? -ne 0 ]; then
+    exit 1
+fi
 
+wait
 
 echo "Running unit tests"
 docker run -a stderr -h robomerge_unittests --name robomerge_unittests robomerge npm test
diff --git a/Engine/Extras/RoboMerge/v3/src/robo/notifications.ts b/Engine/Extras/RoboMerge/v3/src/robo/notifications.ts
--- a/Engine/Extras/RoboMerge/v3/src/robo/notifications.ts
+++ b/Engine/Extras/RoboMerge/v3/src/robo/notifications.ts
@@ -234,6 +234,8 @@ export class SlackMessages {
 			}
 			findResult.messageRecord.messageOpts = message
 			await this.update(findResult.messageRecord, findResult.persistedMessages)
+
+			return findResult.messageRecord.permalink
 		}
 		// Otherwise, we will need to create a new one
 		else {
@@ -257,6 +259,8 @@ export class SlackMessages {
 				findResult.persistedMessages[findResult.conflictKey] = {timestamp, permalink, messageOpts: message}
 				this.persistence.set(SLACK_MESSAGES_PERSISTENCE_KEY, findResult.persistedMessages)
 			}
+
+			return permalink
 		}
 	}
 
@@ -631,6 +635,8 @@ export class BotNotifications implements BotEventHandler {
 			messagesToPost.push({ message: syntaxErrorMessage, reply: true})		
 		}
 
+		let threadLinks = []
+
 		const branchArg = targetBranch ? targetBranch.name : blockage.failure.kind
 		for (const channel of channelsToPostTo) {
 			for (const user of usersToInvite) {
@@ -644,7 +650,7 @@ export class BotNotifications implements BotEventHandler {
 					await this.slackMessages.postReply(changeInfo.cl, branchArg, { ...messageToPost.message, channel })
 				}
 				else {
-					await this.slackMessages.postOrUpdate(changeInfo.cl, branchArg, { ...messageToPost.message, channel })
+					threadLinks.push(await this.slackMessages.postOrUpdate(changeInfo.cl, branchArg, { ...messageToPost.message, channel }))
 				}
 			}
 		}
@@ -665,7 +671,7 @@ export class BotNotifications implements BotEventHandler {
 				const dmText = `Your change (${makeClLink(changeInfo.source_cl)}) ` +
 					`hit '${issue}' while merging from *${sourceBranch.name}* to *${targetBranch.name}*.\n\n` +
 					'`' + blockage.change.description.substring(0, 80) + '`\n\n' +
-					"*_Resolving this blockage is time sensitive._ Please select one of the following:*"
+					"*_Resolving this blockage is time sensitive._*"
 				
 				const urls = this.blockageUrlGenerator(blockage)
 				if (!urls) {
@@ -674,7 +680,7 @@ export class BotNotifications implements BotEventHandler {
 					throw error
 				}
 
-				dm = this.makeSlackDirectMessage(dmText, changeInfo.cl, changeInfo.cl, targetBranch.name, urls)
+				dm = this.makeSlackDirectMessage(dmText, changeInfo.cl, changeInfo.cl, targetBranch.name, urls, threadLinks)
 			}
 
 			this.slackMessages.postDM(userEmail, changeInfo.cl, targetBranch, dm)
@@ -842,10 +848,37 @@ export class BotNotifications implements BotEventHandler {
 
 	// This is an extremely opinionated function to send a stylized direct message to the end user.
 	//makeSlackChannelMessageOpts(title: string, style: string, clDesc: string, sourceBranch: BranchArg,targetBranch?: BranchArg, author?: string, buttons?: SlackLinkButtonsAttachment[]) 
-	private makeSlackDirectMessage(messageText: string, sourceCl: number, conflictCl: number, targetBranch: string, conflictUrls: BlockageNodeOpUrls) : SlackMessage {
+	private makeSlackDirectMessage(messageText: string, sourceCl: number, conflictCl: number, targetBranch: string, conflictUrls: BlockageNodeOpUrls, threadLinks?: string[]) : SlackMessage {
 		// Start collecting our attachments
 		let attachCollection : SlackAttachment[] = []
 		
+		// General information
+		attachCollection.push({
+			text: `To learn more about robomerge and how to resolve blockages, please review the <${this.externalRobomergeUrl}/help|robomerge help page>.`,
+			mrkdwn_in: ["text"]
+		})
+
+		threadLinks = (threadLinks || []).filter(threadLink => threadLink)
+		if (threadLinks.length > 0) {
+			let text = `Discussion for this issue can be found in the following thread${threadLinks.length > 1 ? 's' : ''}:`
+			for (const threadLink of threadLinks) {
+				text += `\n${threadLink}`
+			}
+			attachCollection.push({
+				text,
+				mrkdwn_in: ["text"]
+			})
+		}
+
+		attachCollection.push({
+			text: "You can also get help via the robomerge Slack channel: <#C9321FLTU>",
+			mrkdwn_in: ["text"]
+		})
+		attachCollection.push({
+			text: "If you cannot login to robomerge or access the 'robomerge-help' slack channel, please contact the IT helpdesk.",
+			mrkdwn_in: ["text"]
+		})
+
 		// Acknowledge button
 		const conflictClLink = makeClLink(conflictCl, 'conflict CL #' + conflictCl)
 		attachCollection.push(<SlackLinkButtonsAttachment>{
@@ -885,12 +918,6 @@ export class BotNotifications implements BotEventHandler {
 			})
 		}
 
-		// Append footer
-		attachCollection.push({
-			pretext: "You can get help via the Slack channel <#C044ZUNF9MJ> (if you don't have access to '#dt-robomerge-help-ext', please contact the IT helpdesk)",
-			mrkdwn_in: ["pretext"]
-		})
-
 		// Return SlackMessage for our direct message
 		return {
 			text: messageText,
