Fix crash when previewing or starting the server with an undefined changelist parser

diff --git a/Engine/Extras/RoboMerge/v3/src/robo/branchdefs.ts b/Engine/Extras/RoboMerge/v3/src/robo/branchdefs.ts
--- a/Engine/Extras/RoboMerge/v3/src/robo/branchdefs.ts
+++ b/Engine/Extras/RoboMerge/v3/src/robo/branchdefs.ts
@@ -508,6 +508,15 @@ export class BranchDefs {
 					}
 				}
 			}
+
+			if (def.changelistParser) {
+				try {
+					await import(`./changelistparsers/${def.changelistParser}`)
+				}
+				catch {
+					outErrors.push(`'${def.name}' specifies custom parser '${def.changelistParser}' but unable to import it`)
+				}
+			}
 		}
 
 		// Check branchspecs for valid branches
