Robomerge - Update to Node.JS 20.x LTS
Replace CentOS 7 with AmazonLinux 2023 for docker
#rnx

diff --git a/Engine/Extras/RoboMerge/v3/functional_tests/framework/package.json b/Engine/Extras/RoboMerge/v3/functional_tests/framework/package.json
--- a/Engine/Extras/RoboMerge/v3/functional_tests/framework/package.json
+++ b/Engine/Extras/RoboMerge/v3/functional_tests/framework/package.json
@@ -12,7 +12,7 @@
   "dependencies": {
     "@types/bent": "^7.0.2",
     "@types/mustache": "^0.8.32",
-    "@types/node": "^12.12.34",
+    "@types/node": "^22.9.0",
     "bent": "^7.1.2",
     "colors": "^1.4.0",
     "jsonlint": "^1.6.3",
diff --git a/Engine/Extras/RoboMerge/v3/package.json b/Engine/Extras/RoboMerge/v3/package.json
--- a/Engine/Extras/RoboMerge/v3/package.json
+++ b/Engine/Extras/RoboMerge/v3/package.json
@@ -19,7 +19,7 @@
     "@types/ldapjs": "^1.0.3",
     "@types/marked": "^4.0.2",
     "@types/mustache": "^0.8.32",
-    "@types/node": "^12.12.24",
+    "@types/node": "^22.9.0",
     "@types/nodemailer": "^4.3.4",
     "@types/sanitize-html": "^1.20.2",
     "@types/semver": "^6.0.2",
@@ -42,5 +42,6 @@
   "devDependencies": {
     "@types/intl": "^1.2.0",
     "@types/swagger-ui-express": "^3.0.1"
-  }
+  },
+  "packageManager": "[REDACTED_EMAIL]+sha512.6e2baf77d06b9362294152c851c4f278ede37ab1eba3a55fda317a4a17b209f4dbb973fb250a77abc463a341fcb1f17f17cfa24091c4eb319cda0d9b84278387"
 }
diff --git a/Engine/Extras/RoboMerge/v3/src/common/analytics.ts b/Engine/Extras/RoboMerge/v3/src/common/analytics.ts
--- a/Engine/Extras/RoboMerge/v3/src/common/analytics.ts
+++ b/Engine/Extras/RoboMerge/v3/src/common/analytics.ts
@@ -225,7 +225,7 @@ export class Analytics {
 
 	private perforceRetries = 0
 
-	private flushInterval: NodeJS.Timer	
+	private flushInterval: ReturnType<typeof setTimeout>	
 }
 
 export function runTests(parentLogger: ContextualLogger) {
diff --git a/Engine/Extras/RoboMerge/v3/src/new/graph.ts b/Engine/Extras/RoboMerge/v3/src/new/graph.ts
--- a/Engine/Extras/RoboMerge/v3/src/new/graph.ts
+++ b/Engine/Extras/RoboMerge/v3/src/new/graph.ts
@@ -749,7 +749,7 @@ export function parseTargetsAndFlags(graph: Graph, botname: BotName, tokens: str
 // kind of odd to do this before computing implicit, but shouldn't do any harm
 	for (const auto of automaticTargets) {
 		if (!targets.has(auto)) {
-			targets.set(auto, 'normal' || forcedMode)
+			targets.set(auto, forcedMode ?? 'normal')
 		}
 	}
 
diff --git a/Engine/Extras/RoboMerge/v3/src/robo/graphbot.ts b/Engine/Extras/RoboMerge/v3/src/robo/graphbot.ts
--- a/Engine/Extras/RoboMerge/v3/src/robo/graphbot.ts
+++ b/Engine/Extras/RoboMerge/v3/src/robo/graphbot.ts
@@ -496,7 +496,7 @@ export class GraphBot implements BotEventHandler {
 
 	private _runningBots = false
 	private lastError?: {nodeBot: string, error: string}
-	private timeout: NodeJS.Timer | null = null
+	private timeout: ReturnType<typeof setTimeout> | null = null
 	private _shutdownCb: Function | null = null
 
 	private crashRequested: string | null = null
diff --git a/Engine/Extras/RoboMerge/v3/src/robo/state-interfaces.ts b/Engine/Extras/RoboMerge/v3/src/robo/state-interfaces.ts
--- a/Engine/Extras/RoboMerge/v3/src/robo/state-interfaces.ts
+++ b/Engine/Extras/RoboMerge/v3/src/robo/state-interfaces.ts
@@ -38,7 +38,7 @@ export class PauseState {
 
 	private _infoStr: string
 
-	private blockagePauseTimeout: NodeJS.Timer | null = null;
+	private blockagePauseTimeout: ReturnType<typeof setTimeout> | null = null;
 
 	constructor(private parentUnblock: (reason: string) => void, private readonly context: Context) {
 		this._restore()
diff --git a/Engine/Extras/RoboMerge/v3/src/robo/watchdog.ts b/Engine/Extras/RoboMerge/v3/src/robo/watchdog.ts
--- a/Engine/Extras/RoboMerge/v3/src/robo/watchdog.ts
+++ b/Engine/Extras/RoboMerge/v3/src/robo/watchdog.ts
@@ -379,8 +379,8 @@ class Watchdog {
 
 	private analytics: Analytics
 
-	private respawnTimer: NodeJS.Timer | null = null
-	private memUsageTimer: NodeJS.Timer
+	private respawnTimer: ReturnType<typeof setTimeout> | null = null
+	private memUsageTimer: ReturnType<typeof setTimeout>
 }
 
 const watchdog = new Watchdog
