Use same helix version as is in use at Epic

diff --git a/Engine/Extras/RoboMerge/v3/README.md b/Engine/Extras/RoboMerge/v3/README.md
--- a/Engine/Extras/RoboMerge/v3/README.md
+++ b/Engine/Extras/RoboMerge/v3/README.md
@@ -152,8 +152,8 @@ Dis-allow integration between 5am and 9pm UTC
 The last line of `helix.Dockerfile` contains the version number. The version is also used in `docker-build.sh` and should be updated there:
 
 ```
-HELIX_DOCKER_IMAGE_NAME=${DOCKER_REGISTRY_DOMAIN}/${DOCKER_REGISTRY_NAMESPACE}/helix:2021.1
-HELIX_AND_NODE_DOCKER_IMAGE_NAME=${DOCKER_REGISTRY_DOMAIN}/${DOCKER_REGISTRY_NAMESPACE}/helix-and-node:2021.1
+HELIX_DOCKER_IMAGE_NAME=${DOCKER_REGISTRY_DOMAIN}/${DOCKER_REGISTRY_NAMESPACE}/helix:2022.1
+HELIX_AND_NODE_DOCKER_IMAGE_NAME=${DOCKER_REGISTRY_DOMAIN}/${DOCKER_REGISTRY_NAMESPACE}/helix-and-node:2022.1
 ```
 
 #### NodeJS/Typescript
@@ -169,9 +169,9 @@ Typescript 5.0
 If either of these docker files and modified, a new version of the image should be built and pushed to the Docker Hub. In the Epic internal case, for example:
 
 ```
-docker build -f helix.Dockerfile -t helix
-docker tag helix hub.ol.epicgames.net/epicgames/helix:2021.1
-docker push hub.ol.epicgames.net/epicgames/helix:2021.1
+docker build -f helix.Dockerfile -t helix .
+docker tag helix hub.ol.epicgames.net/epicgames/helix:2022.1
+docker push hub.ol.epicgames.net/epicgames/helix:2022.1
 ```
 
 Note that both the `helix` and `helix-and-node` images use the Helix version as their version.
