[Backout] - CL38871879
#rnx
Original CL Desc
-----------------------------------------------------------------
Use same helix version as is in use at Epic

diff --git a/Engine/Extras/RoboMerge/v3/README.md b/Engine/Extras/RoboMerge/v3/README.md
--- a/Engine/Extras/RoboMerge/v3/README.md
+++ b/Engine/Extras/RoboMerge/v3/README.md
@@ -152,8 +152,8 @@ Dis-allow integration between 5am and 9pm UTC
 The last line of `helix.Dockerfile` contains the version number. The version is also used in `docker-build.sh` and should be updated there:
 
 ```
-HELIX_DOCKER_IMAGE_NAME=${DOCKER_REGISTRY_DOMAIN}/${DOCKER_REGISTRY_NAMESPACE}/helix:2022.1
-HELIX_AND_NODE_DOCKER_IMAGE_NAME=${DOCKER_REGISTRY_DOMAIN}/${DOCKER_REGISTRY_NAMESPACE}/helix-and-node:2022.1
+HELIX_DOCKER_IMAGE_NAME=${DOCKER_REGISTRY_DOMAIN}/${DOCKER_REGISTRY_NAMESPACE}/helix:2021.1
+HELIX_AND_NODE_DOCKER_IMAGE_NAME=${DOCKER_REGISTRY_DOMAIN}/${DOCKER_REGISTRY_NAMESPACE}/helix-and-node:2021.1
 ```
 
 #### NodeJS/Typescript
@@ -169,9 +169,9 @@ Typescript 5.0
 If either of these docker files and modified, a new version of the image should be built and pushed to the Docker Hub. In the Epic internal case, for example:
 
 ```
-docker build -f helix.Dockerfile -t helix .
-docker tag helix hub.ol.epicgames.net/epicgames/helix:2022.1
-docker push hub.ol.epicgames.net/epicgames/helix:2022.1
+docker build -f helix.Dockerfile -t helix
+docker tag helix hub.ol.epicgames.net/epicgames/helix:2021.1
+docker push hub.ol.epicgames.net/epicgames/helix:2021.1
 ```
 
 Note that both the `helix` and `helix-and-node` images use the Helix version as their version.
diff --git a/Engine/Extras/RoboMerge/v3/docker-build.sh b/Engine/Extras/RoboMerge/v3/docker-build.sh
new file mode 100755
--- /dev/null
+++ b/Engine/Extras/RoboMerge/v3/docker-build.sh
@@ -0,0 +1,290 @@
+#!/bin/bash
+
+# TODO: Move script to a build agent image.
+
+#
+# Poor man's docker cleanup. The 'head' command is used to limit how
+# much artifact garbage is collected per invocation. All failures are
+# ignored.
+#
+docker-gc() {
+   exited_containers=$(docker ps -f status=exited -q | head)
+   if [ -n "$exited_containers" ]; then
+      docker rm -v $exited_containers > /dev/null 2>&1
+   fi
+
+   unused_images=$(docker images -q -f dangling=true | head)
+   if [ -n "$unused_images" ]; then
+      docker rmi $unused_images > /dev/null 2>&1
+   fi
+}
+
+die() {
+   err=$?
+   msg=$1
+   if [ -n "$msg" ]; then
+      echo "$msg" >&2
+   fi
+   echo "Failed with exit code $err." >&2
+   exit $err
+}
+
+cleanup() {
+    echo "##teamcity[blockOpened name='Testing cleanup']"
+
+    # Cleanup any old containers
+    docker stop robomerge_unittests
+    docker rm robomerge_unittests
+    docker stop p4docker
+    docker rm p4docker
+    docker stop robomerge_functtest
+    docker rm robomerge_functtest
+    docker stop robomerge_functionaltests
+    docker rm robomerge_functionaltests
+
+    # Cleanup any old networking
+    docker network rm robomerge_functtest_network
+
+    echo "##teamcity[blockClosed name='Testing cleanup']"
+}
+
+#
+# Setup build system supplied environment variables.
+#
+# TODO: Support parsing versions with SNAPSHOT qualifiers that are compatible with Gradle/Ivy/Maven.
+EPIC_BUILD_SOURCE_PATH=$(pwd)
+EPIC_BUILD_SOURCE_BRANCH=${EPIC_BUILD_SOURCE_BRANCH:-UNKNOWN}
+EPIC_BUILD_SOURCE_REVISION=${EPIC_BUILD_SOURCE_REVISION:-UNKNOWN}
+EPIC_BUILD_INSTANCE_NUMBER=${EPIC_BUILD_INSTANCE_NUMBER:-UNKNOWN}
+
+# This nasty sed does the following:
+# 1. We want to capture the Robomerge version from the package.json. We do this for finding the 'version: "#.#.#"' directly following the Robomerge name declaration
+#    This ensures we're capturing the correct version information, and not some random dependancy
+# 2. 'N;  and ;D' surrounding the regex instruct sed to look at the next line in sequence when evaluating the regex -- effective multiline search, needed for #1
+# 3. sed -E uses extended regex, which is better and cleaner than the default. Prevents sed from expecting a lot of escaped characters
+PACKAGE_VERSION=$(sed -E -n '{N; s/.*"name"\s*:\s*"robomerge"\s*,\s+"version":\s*"([^"]+)".*/\1/ip ; D}' package.json)
+EPIC_BUILD_PROJECT_VERSION=${PACKAGE_VERSION:-EPIC_BUILD_PROJECT_VERSION}
+
+EPIC_BUILD_RELEASE_VERSION=${EPIC_BUILD_PROJECT_VERSION}.${EPIC_BUILD_SOURCE_BRANCH}.${EPIC_BUILD_SOURCE_REVISION}.${EPIC_BUILD_INSTANCE_NUMBER}
+
+# Lowercase this so we can use it for the Docker tag as well
+EPIC_BUILD_RELEASE_VERSION=$(echo ${EPIC_BUILD_RELEASE_VERSION} | awk '{ print tolower($0) }')
+echo "##teamcity[buildNumber '$EPIC_BUILD_RELEASE_VERSION']"
+
+DOCKER_REGISTRY_DOMAIN=hub.ol.epicgames.net
+DOCKER_REGISTRY_NAMESPACE=epicgames
+DOCKER_IMAGE_NAME=${DOCKER_REGISTRY_DOMAIN}/${DOCKER_REGISTRY_NAMESPACE}/${EPIC_BUILD_ROLE_NAME}
+echo "##teamcity[setParameter name='env.DOCKER_IMAGE_NAME' value='$DOCKER_IMAGE_NAME']"
+
+HELIX_DOCKER_IMAGE_NAME=${DOCKER_REGISTRY_DOMAIN}/${DOCKER_REGISTRY_NAMESPACE}/helix:2021.1
+HELIX_AND_NODE_DOCKER_IMAGE_NAME=${DOCKER_REGISTRY_DOMAIN}/${DOCKER_REGISTRY_NAMESPACE}/helix-and-node:2021.1
+
+# Fix for TeamCity log reading
+LOGGING_OPTS=
+if [ -n "${EXPLICIT_FUNCTTEST_LOGGING}" ]; then
+    LOGGING_OPTS="--log-driver json-file"
+fi
+
+#
+# Check Docker daemon is running.
+#
+docker version > /dev/null 2>&1 || die "Could not connect to Docker daemon."
+
+#
+# Perform an incremental cleanup up of old docker artifacts before starting build.
+#
+docker-gc
+
+#
+# Build project inside a clean docker build environment sandbox.
+#
+echo "================== Building in sandboxed build environment =================="
+echo "EPIC_BUILD_SOURCE_PATH:     $EPIC_BUILD_SOURCE_PATH"
+echo "EPIC_BUILD_SOURCE_BRANCH:   $EPIC_BUILD_SOURCE_BRANCH"
+echo "EPIC_BUILD_SOURCE_REVISION: $EPIC_BUILD_SOURCE_REVISION"
+echo "EPIC_BUILD_INSTANCE_NUMBER: $EPIC_BUILD_INSTANCE_NUMBER"
+echo "Robomerge PACKAGE_VERSION:  $PACKAGE_VERSION"
+echo "EPIC_BUILD_PROJECT_VERSION: $EPIC_BUILD_PROJECT_VERSION"
+echo "EPIC_BUILD_RELEASE_VERSION: $EPIC_BUILD_RELEASE_VERSION"
+echo "DOCKER_IMAGE_NAME:          $DOCKER_IMAGE_NAME"
+echo "EXPLICIT_FUNCTTEST_LOGGING: $EXPLICIT_FUNCTTEST_LOGGING"
+echo "LOGGING_OPTS:               $LOGGING_OPTS"
+
+# If we have neither a package version nor a EPIC_BUILD_PROJECT_VERSION, exit with error.
+if [ -z "$EPIC_BUILD_PROJECT_VERSION" ]; then
+    echo "Invalid EPIC_BUILD_PROJECT_VERSION! Check package.json or your environment variables."
+    exit 1
+fi
+
+
+#
+# Impersonate the caller's UID/GID so docker builds will
+# create artifacts with the same permissions as the external
+# volume. This is necessary because docker currently does not
+# support user namespace root remapping.
+#
+# see:
+#
+#    [REDACTED_PROFILE_URL]/docker/issues/7906
+#    [REDACTED_PROFILE_URL]/docker/pull/11253
+#
+if [ -n "$(uname -a | grep Linux)" ]; then
+   UID_GID_IMPERSONATION_OPTIONS="--user $(id -u):$(id -g)"
+else
+   #
+   # The Docker daemon on Mac OS X and Windows runs in a lightweight
+   # VM. It automatically mounts and shares the /Users path from the
+   # host's filesystem with docker:staff permissions. This allows the
+   # VM and launched containers to read and write files back to the
+   # host OS.
+   #
+   UID_GID_IMPERSONATION_OPTIONS="--user 1000:50"
+fi
+
+
+#
+# Only build and push runtime images for release target.
+#
+if [ -z "$(echo $@ | awk '{print $0 }' RS=' ' | grep release)" ]; then
+   exit 0
+fi
+
+#
+# Build project runtime docker images.
+#
+echo "##teamcity[blockOpened name='RoboMerge TS Build']"
+
+echo "##teamcity[progressStart 'Downloading helix image and building node...']"
+docker pull ${HELIX_DOCKER_IMAGE_NAME}
+docker tag ${HELIX_DOCKER_IMAGE_NAME} helix
+
+docker pull ${HELIX_AND_NODE_DOCKER_IMAGE_NAME}
+docker tag ${HELIX_AND_NODE_DOCKER_IMAGE_NAME} helix-and-node
+echo "##teamcity[progressFinish 'Downloading helix image and building node...']"
+
+echo "##teamcity[progressStart 'Building RoboMerge TS image...']"
+
+echo "Writing version.json..."
+echo "{ \"build\" : \"${EPIC_BUILD_INSTANCE_NUMBER}\" , \"cl\" : ${EPIC_BUILD_SOURCE_REVISION} }" > version.json
+
+DOCKERFILE=Dockerfile
+echo "Building image ${DOCKER_IMAGE_NAME}:${EPIC_BUILD_RELEASE_VERSION} from ${DOCKERFILE}."
+docker build --file $DOCKERFILE --rm --tag $DOCKER_IMAGE_NAME:$EPIC_BUILD_RELEASE_VERSION --tag $DOCKER_IMAGE_NAME:latest . || die
+
+echo "##teamcity[progressFinish 'Building RoboMerge TS image...']"
+echo "##teamcity[blockClosed name='RoboMerge TS Build']"
+
+echo "##teamcity[blockOpened name='Clean-up before tests']"
+cleanup
+echo "##teamcity[blockClosed name='Clean-up before tests']"
+
+#
+# Run unit tests.
+#
+
+echo "##teamcity[progressStart 'Running $DOCKER_IMAGE_NAME:$EPIC_BUILD_RELEASE_VERSION unit tests...']"
+docker run -a stderr -h robomerge_unittests --name robomerge_unittests $DOCKER_IMAGE_NAME:$EPIC_BUILD_RELEASE_VERSION npm test
+echo "##teamcity[progressFinish 'Running $DOCKER_IMAGE_NAME:$EPIC_BUILD_RELEASE_VERSION unit tests...']"
+
+# Check return value from unit tests
+UNIT_TESTS_RETURN_CODE=`docker inspect robomerge_unittests --format='{{.State.ExitCode}}'`
+if [ "$UNIT_TESTS_RETURN_CODE" != "0" ]; then
+    echo "Unit tests returned a non-zero value ($UNIT_TESTS_RETURN_CODE)! Check logs!"
+    echo "##teamcity[blockOpened name='RoboMerge TS Container Logs']"
+    echo "RoboMerge TS Container Logs:"
+    docker logs robomerge_unittests
+    echo "##teamcity[blockClosed name='RoboMerge TS Container Logs']"
+    cleanup
+    exit 1
+else 
+    echo "Unit tests complete. Return code: $UNIT_TESTS_RETURN_CODE"
+fi
+
+#
+# Run functional tests.
+#
+echo "##teamcity[blockOpened name='Functional Tests Prep']"
+
+# The three containers (P4D, RoboMerge TS, and Functional Test container) need to be on the same docker network
+docker network create --driver bridge robomerge_functtest_network
+
+# P4D
+echo "##teamcity[progressStart 'Starting P4D for functional testing...']"
+docker build --tag p4docker --file Dockerfile.p4docker . || die
+docker run --detach $LOGGING_OPTS --publish 1666:1666 --hostname p4docker --name p4docker --network robomerge_functtest_network p4docker
+sleep 5
+echo "##teamcity[progressFinish 'Starting P4D for functional testing...']"
+
+# Run RoboMerge
+echo "##teamcity[progressStart 'Starting $DOCKER_IMAGE_NAME:$EPIC_BUILD_RELEASE_VERSION for functional testing...']"
+docker run --detach $LOGGING_OPTS --publish 8877:8877 --publish 8811:8811 --env P4PORT=p4docker:1666 --hostname robomerge_functtest --name robomerge_functtest  \
+            -e ROBO_DEV_MODE=true                                 \
+            -e ROBO_DEV_MODE_USER=testuser1                       \
+            -e ROBO_USE_SLACK_IN_DEV=true                         \
+            -e P4PORT=p4docker:1666                               \
+            -e ROBO_BRANCHSPECS_ROOT_PATH=//RoboMergeData/Main    \
+            -e ROBO_PERSISTENCE_BACKUP_FREQUENCY=0                \
+            -e ROBO_NO_TLS=true                                   \
+            -e ROBO_EXTERNAL_URL=http://robomerge_functtest:8877  \
+            -e BOTNAME=ft1,ft2,ft3,ft4,targets                    \
+            -e ROBO_LOG_LEVEL=info                                \
+            -e ROBO_SLACK_DOMAIN=http://localhost:8811            \
+            --network robomerge_functtest_network $DOCKER_IMAGE_NAME:$EPIC_BUILD_RELEASE_VERSION node dist/robo/watchdog.js
+echo "##teamcity[progressFinish 'Starting $DOCKER_IMAGE_NAME:$EPIC_BUILD_RELEASE_VERSION for functional testing...']"
+
+# Build functional test image
+echo "##teamcity[progressStart 'Building functional tests image...']"
+docker build --tag robomerge_functionaltests --file functional_tests/tstests.Dockerfile . || die
+echo "##teamcity[progressFinish 'Building functional tests image...']"
+echo "##teamcity[blockClosed name='Functional Tests Prep']"
+
+
+# Run Functional Tests
+echo "##teamcity[blockOpened name='Functional Tests']"
+echo "##teamcity[progressStart 'Running functional tests...']"
+docker run --hostname robomerge_functionaltests --name robomerge_functionaltests --network robomerge_functtest_network robomerge_functionaltests
+
+# Check return value from Functional Tests
+FUNCT_TEST_RETURN_CODE=`docker inspect robomerge_functionaltests --format='{{.State.ExitCode}}'`
+if [ "$FUNCT_TEST_RETURN_CODE" != "0" ]; then
+    echo "Functional tests returned a non-zero value ($FUNCT_TEST_RETURN_CODE)! Check logs!"
+else
+    echo "Functional tests complete. Return code: $FUNCT_TEST_RETURN_CODE"
+fi
+
+echo "##teamcity[blockOpened name='RoboMerge Container Logs']"
+echo "P4Docker Container Logs:"
+docker logs p4docker
+echo "RoboMerge TS Container Logs:"
+docker logs robomerge_functtest
+echo "##teamcity[blockClosed name='RoboMerge Container Logs']"
+
+if [ "$FUNCT_TEST_RETURN_CODE" != "0" ]; then
+    cleanup
+    exit 1
+fi
+
+echo "##teamcity[progressFinish 'Running functional tests...']"
+echo "##teamcity[blockClosed name='Functional Tests']"
+
+# After functional tests have verified the build, push to DockerHub
+echo "##teamcity[blockOpened name='Docker Push']"
+echo "##teamcity[progressStart 'Pushing RoboMerge TS image...']"
+docker push $DOCKER_IMAGE_NAME:$EPIC_BUILD_RELEASE_VERSION || die
+docker push $DOCKER_IMAGE_NAME:latest || die
+echo "##teamcity[progressFinish 'Pushing RoboMerge TS image...']"
+echo "##teamcity[blockClosed name='Docker Push']"
+
+echo "##teamcity[blockOpened name='Clean-up after tests']"
+cleanup
+echo "##teamcity[blockClosed name='Clean-up after tests']"
+
+
+#
+# TeamCity hook to update variables based on the outcome of this build. Required
+# to propagate values through a build chain.
+#
+# see: https://confluence.jetbrains.com/display/TCD7/Build+Script+Interaction+with+TeamCity
+#
+echo "##teamcity[setParameter name='env.EPIC_BUILD_RELEASE_VERSION' value='$EPIC_BUILD_RELEASE_VERSION']"
+echo "##teamcity[buildStatus text='{build.status.text} - $DOCKER_IMAGE_NAME:$EPIC_BUILD_RELEASE_VERSION']"
