Update graphviz implementation

diff --git a/Engine/Extras/RoboMerge/v3/install-graphviz.js b/Engine/Extras/RoboMerge/v3/install-graphviz.js
--- a/Engine/Extras/RoboMerge/v3/install-graphviz.js
+++ b/Engine/Extras/RoboMerge/v3/install-graphviz.js
@@ -2,19 +2,10 @@
 // Copyright Epic Games, Inc. All Rights Reserved.
 
 const fs = require('fs')
-const gzip = require('node-gzip').gzip
 
 Promise.all([
-	fs.promises.copyFile('node_modules/@hpcc-js/wasm/dist/index.min.js', 'public/js/hpcc-wasm-index.min.js'),
-	Promise.all([
-		fs.promises.mkdir('public/bin').catch(() => {}),
-		fs.promises.readFile('node_modules/@hpcc-js/wasm/dist/graphvizlib.wasm')
-		.then(gzip)
-	])
-	.then(([_, compressed]) => {
-		fs.promises.writeFile('public/bin/graphvizlib.wasm.gz', compressed);
-	})
+	fs.promises.copyFile('node_modules/@viz-js/viz/lib/viz-standalone.js', 'public/js/viz-standalone.js'),
 ])
 .then(() => {
-	console.log('Installed GraphViz binary for web app')
+	console.log('Installed GraphViz for web app')
 })
\ No newline at end of file
diff --git a/Engine/Extras/RoboMerge/v3/package.json b/Engine/Extras/RoboMerge/v3/package.json
--- a/Engine/Extras/RoboMerge/v3/package.json
+++ b/Engine/Extras/RoboMerge/v3/package.json
@@ -13,8 +13,6 @@
     "test": "node dist/robo/unit-tests.js"
   },
   "dependencies": {
-    "@hpcc-js/graph": "^2.27.0",
-    "@hpcc-js/wasm": "^0.3.14",
     "@sentry/node": "^5.7.1",
     "@types/ldapjs": "^1.0.3",
     "@types/marked": "^4.0.2",
@@ -24,6 +22,7 @@
     "@types/sanitize-html": "^1.20.2",
     "@types/semver": "^6.0.2",
     "@types/swagger-ui-dist": "^3.0.3",
+    "@viz-js/viz": "^3.11.0",
     "colors": "^1.4.0",
     "jsonlint": "^1.6.3",
     "ldapauth-fork": "^4.0.2",
diff --git a/Engine/Extras/RoboMerge/v3/public/allbots.html b/Engine/Extras/RoboMerge/v3/public/allbots.html
--- a/Engine/Extras/RoboMerge/v3/public/allbots.html
+++ b/Engine/Extras/RoboMerge/v3/public/allbots.html
@@ -11,7 +11,7 @@
 
 	<script src="/js/jquery-3.4.1.min.js"></script>
 	<script src="/js/bootstrap.bundle.min.js"></script>
-	<script src="/js/hpcc-wasm-index.min.js"></script>
+	<script src="/js/viz-standalone.js"></script>
 	<script src="/js/flow.graph.js"></script>
 	<script src="/js/flow.js"></script>
 	<script src="/js/botselect.js"></script>
diff --git a/Engine/Extras/RoboMerge/v3/public/index.html b/Engine/Extras/RoboMerge/v3/public/index.html
--- a/Engine/Extras/RoboMerge/v3/public/index.html
+++ b/Engine/Extras/RoboMerge/v3/public/index.html
@@ -20,7 +20,7 @@
 	<script src="/js/bootstrap.bundle.min.js"></script>
 
 	<script src="/js/boilerplate.js"></script>
-	<script src="/js/hpcc-wasm-index.min.js"></script>
+	<script src="/js/viz-standalone.js"></script>
 	<script src="/js/flow.graph.js"></script>
 	<script src="/js/flow.js"></script>
 	<script src="/js/main.js"></script>
diff --git a/Engine/Extras/RoboMerge/v3/public/js/flow.js b/Engine/Extras/RoboMerge/v3/public/js/flow.js
--- a/Engine/Extras/RoboMerge/v3/public/js/flow.js
+++ b/Engine/Extras/RoboMerge/v3/public/js/flow.js
@@ -15,9 +15,11 @@ function setDefault(map, key, def) {
     map.set(key, def);
     return def;
 }
+
 function renderGraph(src) {
-    const hpccWasm = window["@hpcc-js/wasm"];
-    return hpccWasm.graphviz.layout(src, "svg", "dot");
+    return Viz.instance().then(function (viz) {
+        return viz.renderSVGElement(src);
+    })
 }
 
 function parseOptions(search) {
diff --git a/Engine/Extras/RoboMerge/v3/public/preview.html b/Engine/Extras/RoboMerge/v3/public/preview.html
--- a/Engine/Extras/RoboMerge/v3/public/preview.html
+++ b/Engine/Extras/RoboMerge/v3/public/preview.html
@@ -7,7 +7,7 @@
 
 	<script src="/js/jquery-3.4.1.min.js"></script>
 	<script src="/js/bootstrap.bundle.min.js"></script>
-	<script src="/js/hpcc-wasm-index.min.js"></script>
+	<script src="/js/viz-standalone.js"></script>
 	<script src="/js/flow.graph.js"></script>
 	<script src="/js/flow.js"></script>
 	<script src="/js/botselect.js"></script>
diff --git a/Engine/Extras/RoboMerge/v3/public/trackchange.html b/Engine/Extras/RoboMerge/v3/public/trackchange.html
--- a/Engine/Extras/RoboMerge/v3/public/trackchange.html
+++ b/Engine/Extras/RoboMerge/v3/public/trackchange.html
@@ -7,7 +7,7 @@
 
 	<script src="/js/jquery-3.4.1.min.js"></script>
 	<script src="/js/bootstrap.bundle.min.js"></script>
-	<script src="/js/hpcc-wasm-index.min.js"></script>
+	<script src="/js/viz-standalone.js"></script>
 	<script src="/js/flow.graph.js"></script>
 	<script src="/js/flow.js"></script>
 	<script src="/js/trackchange.js"></script>
diff --git a/Engine/Extras/RoboMerge/v3/src/robo/roboserver.ts b/Engine/Extras/RoboMerge/v3/src/robo/roboserver.ts
--- a/Engine/Extras/RoboMerge/v3/src/robo/roboserver.ts
+++ b/Engine/Extras/RoboMerge/v3/src/robo/roboserver.ts
@@ -137,13 +137,6 @@ export class RoboServer {
 		}
 
 		this.server.addFileMapping('/allbots', 'allbots.html')
-		this.server.addFileMapping('/js/*.wasm', 'bin/$1.wasm.gz', {
-			filetype: "application/wasm", 
-			headers: [
-				['Cache-Control', 'max-age=' + 60*60*24*7],
-				['Content-Encoding', 'gzip']
-			]
-			})
 		this.server.addFileMapping('/js/*.js', 'js/$1.js')
 		this.server.addFileMapping('/css/*.css', 'css/$1.css')
 		this.server.addFileMapping('/img/*.png', 'images/$1.png')
