Adapt to deprecation of files.upload

diff --git a/Engine/Extras/RoboMerge/v3/src/robo/slack.ts b/Engine/Extras/RoboMerge/v3/src/robo/slack.ts
--- a/Engine/Extras/RoboMerge/v3/src/robo/slack.ts
+++ b/Engine/Extras/RoboMerge/v3/src/robo/slack.ts
@@ -128,7 +128,32 @@ export class Slack {
 	}
 
 	uploadFile(file: SlackFile) {
-		return this.post('files.upload', file, false, PostFormat.URL)
+		let getURLargs: any = {
+			filename: file.filename,
+			length: file.content.length
+		}
+		if (file.filetype) {
+			getURLargs.snippet_type = file.filetype // slack API gets snippy if you pass in snippet_type = undefined
+		}
+		return this.post('files.getUploadURLExternal', getURLargs, false, PostFormat.URL)
+			.then( getURLresult => {
+				request.post({
+					url: getURLresult.upload_url,
+					body: file.content,
+					contentType: 'application/x-www-form-urlencoded'
+				}).then(
+					() => {
+						const completeArgs = {
+							files:[{id: getURLresult.file_id, title: file.title}],
+							channel_id: file.channels,
+							initial_comment: file.initial_comment,
+							thread_ts: file.thread_ts
+						}
+						this.post('files.completeUploadExternal', completeArgs)
+					}
+				)
+			}
+		)
 	}
 
 	listMessages(count?: number) {
