Directly ping owner in addition to triager when nagging

diff --git a/Engine/Extras/RoboMerge/v3/src/common/perforce.ts b/Engine/Extras/RoboMerge/v3/src/common/perforce.ts
--- a/Engine/Extras/RoboMerge/v3/src/common/perforce.ts
+++ b/Engine/Extras/RoboMerge/v3/src/common/perforce.ts
@@ -2,7 +2,7 @@
 import { execFile, ExecFileOptions } from 'child_process';
 import * as path from 'path';
 // should really separate out RM-specific operations, but no real gain at the moment
-import { isUserAKnownBot, postToRobomergeAlerts } from '../robo/notifications';
+import { isUserAKnownBot, isUserSlackGroup, postToRobomergeAlerts } from '../robo/notifications';
 import { roboAnalytics } from '../robo/roboanalytics';
 import { ContextualLogger, NpmLogLevel } from './logger';
 import { Vault } from '../robo/vault'
@@ -1601,7 +1601,7 @@ export class PerforceContext {
 	}
 
 	async getUser(username: string) {
-		if (username.length > 0 && !username.startsWith('@') && !username.startsWith('<') && !isUserAKnownBot(username)) {
+		if (username.length > 0 && !isUserSlackGroup(username) && !isUserAKnownBot(username)) {
 			try {
 				return (await this.execAndParse(null, ['users', username], {quiet:true}))[0]
 			}
diff --git a/Engine/Extras/RoboMerge/v3/src/robo/nodebot.ts b/Engine/Extras/RoboMerge/v3/src/robo/nodebot.ts
--- a/Engine/Extras/RoboMerge/v3/src/robo/nodebot.ts
+++ b/Engine/Extras/RoboMerge/v3/src/robo/nodebot.ts
@@ -14,7 +14,7 @@ import { MergeAction, OperationResult, PendingChange, resolveBranchArg, StompedR
 import { Conflicts } from './conflicts';
 import { EdgeBot, EdgeMergeResults } from './edgebot';
 import { BotEventTriggers } from './events';
-import { makeClLink, SlackMessages } from './notifications';
+import { getPingTarget, isUserSlackGroup, makeClLink, SlackMessages } from './notifications';
 import { PerforceStatefulBot } from './perforce-stateful-bot';
 import { BlockageNodeOpUrls, OperationUrlHelper } from './roboserver';
 import { Context } from './settings';
@@ -2340,7 +2340,7 @@ export class NodeBot extends PerforceStatefulBot implements NodeBotInterface {
 				
 						let triager = getNagProperty("triager")
 
-						if (triager && !triager.startsWith('@') && !triager.startsWith('<'))
+						if (triager && !isUserSlackGroup(triager))
 						{
 							const emailAddress = await this.findEmail(triager)
 							if (emailAddress)
@@ -2360,14 +2360,16 @@ export class NodeBot extends PerforceStatefulBot implements NodeBotInterface {
 							}
 						}
 
+						const channelPing = await getPingTarget(await this.p4.getEmail(conflict.owner), conflict.owner, this.slackMessages)
+
 						const nagMessage: SlackMessage = {
-							text: `${triager ? triager + ' ' : ''}RoboMerge blocked for more than ${timeDesc}`,
+							text: `${triager ? triager + ' ' : ''}${channelPing} RoboMerge blocked for more than ${timeDesc}`,
 							style: SlackMessageStyles.DANGER,
 							channel: this.branchGraph.config.slackChannel,
 							mrkdwn: true
 						}
 
-						this.nodeBotLogger.info(`Sending nag notification to ${conflict.owner} after ${Math.floor(ageMinutes)} minutes`)
+						this.nodeBotLogger.info(`Sending nag notification to ${triager ? triager + ' and ' : ''}${conflict.owner} after ${Math.floor(ageMinutes)} minutes`)
 						const branchName =  (conflict.kind === "Syntax error" ? conflict.kind : conflict.targetBranchName || conflict.blockedBranchName)
 						this.slackMessages.postReply(conflict.cl, branchName, nagMessage)
 					}
diff --git a/Engine/Extras/RoboMerge/v3/src/robo/notifications.ts b/Engine/Extras/RoboMerge/v3/src/robo/notifications.ts
--- a/Engine/Extras/RoboMerge/v3/src/robo/notifications.ts
+++ b/Engine/Extras/RoboMerge/v3/src/robo/notifications.ts
@@ -77,10 +77,23 @@ export async function postMessageToChannel(message: string, channel: string, sty
 //////////
 // Utils
 
-// make a Slack notifcation link for a user
+// make a Slack notification link for a user
 function atifyUser(user: string) {
 	// don't @ names of bots (currently making them tt style for Slack)
-	return isUserAKnownBot(user) ? `\`${user}\`` : user.startsWith('@') || user.startsWith('<') ? user : '@' + user
+	return isUserAKnownBot(user) ? `\`${user}\`` : isUserSlackGroup(user) ? user : '@' + user
+}
+
+export function isUserSlackGroup(user: string) {
+	return user.startsWith('@') || user.startsWith('<')
+}
+
+export async function getPingTarget(userEmail: string|null, owner: string, slackMessages: SlackMessages) {
+	const slackUser = userEmail ? await slackMessages.getSlackUser(userEmail) : null
+	return slackUser 
+	? `<@${slackUser}>` 
+	: isUserSlackGroup(owner)
+		? owner 
+		: `@${owner}`
 }
 
 function generatePersistedSlackMessageKey(sourceCl: number, targetBranchArg: BranchArg, channel: string) {
@@ -528,12 +541,7 @@ export class BotNotifications implements BotEventHandler {
 		const issue = blockage.failure.kind.toLowerCase()
 
 		const userEmail = await blockage.ownerEmail
-		const slackUser = userEmail ? await this.slackMessages.getSlackUser(userEmail) : null
-		const channelPing = slackUser 
-		                        ? `<@${slackUser}>` 
-		                        : blockage.owner.startsWith('@') || blockage.owner.startsWith('<')
-		                              ? blockage.owner 
-		                              : `@${blockage.owner}`
+		const channelPing = await getPingTarget(userEmail, blockage.owner, this.slackMessages)
 
 		const isBotUser = isUserAKnownBot(blockage.owner)
 		const text =
