Fix crash bot functionality

diff --git a/Engine/Extras/RoboMerge/v3/src/robo/graphbot.ts b/Engine/Extras/RoboMerge/v3/src/robo/graphbot.ts
--- a/Engine/Extras/RoboMerge/v3/src/robo/graphbot.ts
+++ b/Engine/Extras/RoboMerge/v3/src/robo/graphbot.ts
@@ -22,6 +22,7 @@ import { Status } from './status';
 import { GraphBotState } from "./status-types"
 import { TickJournal } from './tick-journal';
 import { GraphAPI } from '../new/graph';
+import * as p4util from '../common/p4util';
 
 // probably get the gist after 2000 characters
 const MAX_ERROR_LENGTH_TO_REPORT = 2000
@@ -178,6 +179,7 @@ export class GraphBot implements BotEventHandler {
 		}
 
 		delete this.lastError
+		GraphBot.crashRequests.delete(this.branchGraph.botname)
 
 		const msg = `${who} restarted bot ${this.branchGraph.botname}`
 		this.botLogger.info(msg)
@@ -186,12 +188,20 @@ export class GraphBot implements BotEventHandler {
 		await this.startBotsAsync()
 	}
 
+	static checkCrashRequest(botname: string) {
+		// crashMe API support - simulate a bot crashing and stopping the GraphBot instance
+		const crashRequested = this.crashRequests.get(botname)
+		if (crashRequested) {
+			throw new Error(crashRequested)
+		}
+	}
+
 	// Don't call this unless you want to bring down the entire GraphBot in a crash!
 	async danger_crashGraphBot(who: string) {
 		const msg = `${who} has requested a crash for bot ${this.branchGraph.botname}`
 		this.botLogger.warn(msg)
 		await postMessageToChannel(msg, this.branchGraph.config.slackChannel)
-		this.crashRequested = msg
+		GraphBot.crashRequests.set(this.branchGraph.botname, msg)
 	}
 
 	async handleRequestedIntegrationsForAllNodes(): Promise<[NodeBot, Error] | null> {
@@ -270,16 +280,19 @@ export class GraphBot implements BotEventHandler {
 
 			const startTime = Date.now()
 
-			let tickBot = async function(bot: Bot, crashRequested: string|null): Promise<[Bot,boolean|Error]> {		
+			let tickBot = async function(bot: Bot): Promise<[Bot,boolean|Error]> {		
 				bot.isActive = true
 				let ticked = false
 				try {
-					// crashMe API support - simulate a bot crashing and stopping the GraphBot instance
-					if (crashRequested) {
-						throw new Error(crashRequested)
+					if (bot instanceof NodeBot) {
+						GraphBot.checkCrashRequest(bot.branchGraph.botname)
 					}
 
 					ticked = await bot.tick()
+
+					if (bot instanceof NodeBot) {
+						GraphBot.checkCrashRequest(bot.branchGraph.botname)
+					}
 				}
 				catch (err) {
 					return [bot,err]
@@ -298,7 +311,7 @@ export class GraphBot implements BotEventHandler {
 
 			// The autoupdater needs to fully run before we parallelize the remaining bots
 			if (this.autoUpdater) {
-				const [_, autoUpdateResult] = await tickBot(this.autoUpdater, null)
+				const [_, autoUpdateResult] = await tickBot(this.autoUpdater)
 				if (typeof autoUpdateResult !== 'boolean') {
 					this.handleNodebotError(this.autoUpdater, autoUpdateResult)
 					return
@@ -312,9 +325,7 @@ export class GraphBot implements BotEventHandler {
 				}
 			}
 
-			const tickResults = await Promise.all(this.botlist.map(async (bot, index) => tickBot(bot, index == 0 ? this.crashRequested : null)));
-
-			this.crashRequested = null
+			const tickResults = await Promise.all(this.botlist.map(async (bot) => tickBot(bot)));
 			
 			let botCrashed = false
 			for (const [bot, result] of tickResults)
@@ -499,5 +510,5 @@ export class GraphBot implements BotEventHandler {
 	private timeout: ReturnType<typeof setTimeout> | null = null
 	private _shutdownCb: Function | null = null
 
-	private crashRequested: string | null = null
+	private static crashRequests: Map<string, string> = new Map()
 }
diff --git a/Engine/Extras/RoboMerge/v3/src/robo/nodebot.ts b/Engine/Extras/RoboMerge/v3/src/robo/nodebot.ts
--- a/Engine/Extras/RoboMerge/v3/src/robo/nodebot.ts
+++ b/Engine/Extras/RoboMerge/v3/src/robo/nodebot.ts
@@ -14,6 +14,7 @@ import { MergeAction, OperationResult, PendingChange, resolveBranchArg, StompedR
 import { Conflicts } from './conflicts';
 import { EdgeBot, EdgeMergeResults } from './edgebot';
 import { BotEventTriggers } from './events';
+import { GraphBot } from './graphbot';
 import { getPingTarget, isUserSlackGroup, makeClLink, SlackMessages } from './notifications';
 import { PerforceStatefulBot } from './perforce-stateful-bot';
 import { BlockageNodeOpUrls, OperationUrlHelper } from './roboserver';
@@ -1568,6 +1569,8 @@ export class NodeBot extends PerforceStatefulBot implements NodeBotInterface {
 				}
 			}
 
+			GraphBot.checkCrashRequest(this.graphBotName)
+
 			const changeResult = await this._processAndMergeCl(availableEdges, change, false)
 
 			// Exit immediately on syntax errors
