Properly clean workspaces on restart

diff --git a/Engine/Extras/RoboMerge/v3/src/common/p4util.ts b/Engine/Extras/RoboMerge/v3/src/common/p4util.ts
--- a/Engine/Extras/RoboMerge/v3/src/common/p4util.ts
+++ b/Engine/Extras/RoboMerge/v3/src/common/p4util.ts
@@ -23,8 +23,16 @@ function getWorkspaceID(serverAddress: string | undefined, workspaceName: string
 }
 // true has been cleaned, false is being cleaned, undefined needs cleaning
 let cleanedWorkspaces: Map<WorkspaceID, boolean> = new Map()
+let workspacesByBot: Map<string, Set<WorkspaceID>> = new Map()
 
-export async function cleanWorkspace(logger: ContextualLogger, p4: PerforceContext, workspace: Workspace, workspaceRoot: string, edgeServerAddress?: string) {
+export async function clearCleanedWorkspaces(botName: string) {
+	for (const workspaceID of (workspacesByBot.get(botName) || [])) {
+		cleanedWorkspaces.delete(workspaceID)
+	}
+	workspacesByBot.delete(botName)
+}
+
+export async function cleanWorkspace(logger: ContextualLogger, p4: PerforceContext, botName: string, workspace: Workspace, workspaceRoot: string, edgeServerAddress?: string) {
 	
 	const workspaceID: WorkspaceID = getWorkspaceID(p4.serverAddress, workspace.name)
 
@@ -42,6 +50,7 @@ export async function cleanWorkspace(logger: ContextualLogger, p4: PerforceConte
 
 	logger.info(`Cleaning ${workspace.name}`)
 	cleanedWorkspaces.set(workspaceID, false)
+	workspacesByBot.set(botName, (workspacesByBot.get(botName) || new Set<string>()).add(workspaceID))
 	const changes = await p4.get_pending_changes(workspace.name) as Change[]
 
 	// revert any pending changes left over from previous runs
diff --git a/Engine/Extras/RoboMerge/v3/src/robo/autobranchupdater.ts b/Engine/Extras/RoboMerge/v3/src/robo/autobranchupdater.ts
--- a/Engine/Extras/RoboMerge/v3/src/robo/autobranchupdater.ts
+++ b/Engine/Extras/RoboMerge/v3/src/robo/autobranchupdater.ts
@@ -224,7 +224,7 @@ export class AutoBranchUpdater implements Bot {
 			if (workspaceQueryResult[0].IsUnloaded) {
 				await this.p4.reloadWorkspace(workspace.name)
 			}
-			await p4util.cleanWorkspace(this.abuLogger, this.p4, workspace, workspace.depotpath)
+			await p4util.cleanWorkspace(this.abuLogger, this.p4, this.graphBot.branchGraph.botname, workspace, workspace.depotpath)
 		}
 
 		const {depotpath, realFilepath, mirrorFilepath} = workspace
diff --git a/Engine/Extras/RoboMerge/v3/src/robo/edgebot.ts b/Engine/Extras/RoboMerge/v3/src/robo/edgebot.ts
--- a/Engine/Extras/RoboMerge/v3/src/robo/edgebot.ts
+++ b/Engine/Extras/RoboMerge/v3/src/robo/edgebot.ts
@@ -507,7 +507,7 @@ class EdgeBotImpl extends PerforceStatefulBot {
 			if (existingWorkspaceInfo[0].IsUnloaded) {
 				await p4.reloadWorkspace(workspaceName, edgeServerAddress)
 			}
-			await p4util.cleanWorkspace(this.logger, p4, {name: workspaceName, directory: path}, this.targetBranch.rootPath, edgeServerAddress)
+			await p4util.cleanWorkspace(this.logger, p4, this.graphBotName, {name: workspaceName, directory: path}, this.targetBranch.rootPath, edgeServerAddress)
 		}
 		else {
 			const params: any = {};
diff --git a/Engine/Extras/RoboMerge/v3/src/robo/graphbot.ts b/Engine/Extras/RoboMerge/v3/src/robo/graphbot.ts
--- a/Engine/Extras/RoboMerge/v3/src/robo/graphbot.ts
+++ b/Engine/Extras/RoboMerge/v3/src/robo/graphbot.ts
@@ -185,6 +185,9 @@ export class GraphBot implements BotEventHandler {
 		this.botLogger.info(msg)
 		postToRobomergeAlerts(msg)
 
+		this.botLogger.info('Marking workspaces to be recleaned')
+		await p4util.clearCleanedWorkspaces(this.branchGraph.botname)
+
 		await this.startBotsAsync()
 	}
 
