Adapt dummy slack server for new file upload API

diff --git a/Engine/Extras/RoboMerge/v3/src/common/dummyslackserver.ts b/Engine/Extras/RoboMerge/v3/src/common/dummyslackserver.ts
--- a/Engine/Extras/RoboMerge/v3/src/common/dummyslackserver.ts
+++ b/Engine/Extras/RoboMerge/v3/src/common/dummyslackserver.ts
@@ -102,9 +102,10 @@ export class DummySlackApp implements AppInterface {
 			const dataObj = dataResult.dataObj
 			return { ok: true, channel: {id: `${dataObj.users}`} }
 		}
-		else if (command === "files.upload") {
+		else if (command === "files.getUploadURLExternal" ||
+				 command === "files.completeUploadExternal") {
 			// Not going to store files until we have a reason to do so
-			return { ok: true }
+			return { ok: false }
 		}
 		else {
 			console.error(`dummy-slack POST: unsupported command '${command}'`)
diff --git a/Engine/Extras/RoboMerge/v3/src/robo/slack.ts b/Engine/Extras/RoboMerge/v3/src/robo/slack.ts
--- a/Engine/Extras/RoboMerge/v3/src/robo/slack.ts
+++ b/Engine/Extras/RoboMerge/v3/src/robo/slack.ts
@@ -135,23 +135,25 @@ export class Slack {
 		if (file.filetype) {
 			getURLargs.snippet_type = file.filetype // slack API gets snippy if you pass in snippet_type = undefined
 		}
-		return this.post('files.getUploadURLExternal', getURLargs, false, PostFormat.URL)
+		return this.post('files.getUploadURLExternal', getURLargs, true, PostFormat.URL)
 			.then( getURLresult => {
-				request.post({
-					url: getURLresult.upload_url,
-					body: file.content,
-					contentType: 'application/x-www-form-urlencoded'
-				}).then(
-					() => {
-						const completeArgs = {
-							files:[{id: getURLresult.file_id, title: file.title}],
-							channel_id: file.channels,
-							initial_comment: file.initial_comment,
-							thread_ts: file.thread_ts
+				if (getURLresult.ok) {
+					request.post({
+						url: getURLresult.upload_url,
+						body: file.content,
+						contentType: 'application/x-www-form-urlencoded'
+					}).then(
+						() => {
+							const completeArgs = {
+								files:[{id: getURLresult.file_id, title: file.title}],
+								channel_id: file.channels,
+								initial_comment: file.initial_comment,
+								thread_ts: file.thread_ts
+							}
+							this.post('files.completeUploadExternal', completeArgs)
 						}
-						this.post('files.completeUploadExternal', completeArgs)
-					}
-				)
+					)
+				}
 			}
 		)
 	}
