Fix reload workspace crashing graphbot

diff --git a/Engine/Extras/RoboMerge/v3/src/common/perforce.ts b/Engine/Extras/RoboMerge/v3/src/common/perforce.ts
--- a/Engine/Extras/RoboMerge/v3/src/common/perforce.ts
+++ b/Engine/Extras/RoboMerge/v3/src/common/perforce.ts
@@ -816,8 +816,25 @@ export class PerforceContext {
 		return null
 	}
 
-	reloadWorkspace(workspaceName: string, edgeServerAddress?: string) {
-		return this.execAndParse(null, ['reload', '-c', workspaceName], {serverAddress: edgeServerAddress});
+	async reloadWorkspace(workspaceName: string, edgeServerAddress?: string) {
+		let result
+		try {
+			 result = await this.execAndParse(null, ['reload', '-c', workspaceName], {serverAddress: edgeServerAddress});
+		}
+		catch (reason) {
+			if (!isExecP4Error(reason)) {
+				throw reason
+			}
+			const [err, output] = reason
+
+			// reload succeeds with an error because ... sigh
+			if (!output.trim().startsWith('Client') ||
+				!output.trim().endsWith('reloaded.')) {
+				throw err
+			}
+			result = output;			
+		}
+		return result
 	}
 
 	getEdgeServerAddress(serverId: string) {
