Retry failed print commands a few times and then treat it as a transfer error.

diff --git a/Engine/Extras/RoboMerge/v3/src/common/perforce.ts b/Engine/Extras/RoboMerge/v3/src/common/perforce.ts
--- a/Engine/Extras/RoboMerge/v3/src/common/perforce.ts
+++ b/Engine/Extras/RoboMerge/v3/src/common/perforce.ts
@@ -1134,15 +1134,32 @@ export class PerforceContext {
 
 		await Promise.all(promises)
 
+		const MAX_PRINT_FAILURES = 10
+		let failedPrints = 0
 		if (failures.length == 0) {
 			const PARALLEL_PRINTS = 25
 			const doPrints = async (startIndex: number) => {
 				for (let i=startIndex; i < printCommands.length; i += PARALLEL_PRINTS) {
-					await this.print(printCommands[i][0],printCommands[i][1])
+					while (failedPrints < MAX_PRINT_FAILURES) {
+						try {
+							await this.print(printCommands[i][0],printCommands[i][1])
+							break
+						} catch(reason) {
+							if (!isExecP4Error(reason)) {
+								throw reason
+							}
+							const [err, _] = reason
+							this.logger.warn(err)
+							failedPrints++
+						}
+					}
 				}
 			}
 			await Promise.all(Array.from(Array(PARALLEL_PRINTS).keys()).map(n => doPrints(n)))
 		}
+		if (failedPrints == MAX_PRINT_FAILURES) {
+			failures.push(`Failed p4 print ${MAX_PRINT_FAILURES} times. Aborting. See log for failures.`)
+		}
 
 		return failures
 	}
