Remove remaining vestiges of edigrate functionality
Add a skip if unchanged integration method that can be specified per edge

diff --git a/Engine/Extras/RoboMerge/v3/README.md b/Engine/Extras/RoboMerge/v3/README.md
--- a/Engine/Extras/RoboMerge/v3/README.md
+++ b/Engine/Extras/RoboMerge/v3/README.md
@@ -102,7 +102,7 @@ for both node and edge.
 |`forcePause`               |ne|            | |If flag set, pause - applies each restart| |
 |`additionalSlackChannelForBlockages` |e |  | |Single extra Slack channel               | |
 |`blockAssetFlow`           |n |`[]`        | |Reject any integrations containing assets to these targets   | |
-|`defaultIntegrationMethod` |e |normal      | |For edigrate                             | |
+|`integrationMethod`        |e |normal      | |For skip if unchanged                    | |
 |`disallowSkip`             |e |            | |Remove skip option from UI               | |
 |`emailOnBlockage`          |e |`true`      | |Email owners of conflicts?               | |
 |`excludeAuthors`           |e |            | |Skip changes by these authors, e.g. skip `buildmachine` commits| |
diff --git a/Engine/Extras/RoboMerge/v3/functional_tests/framework/src/framework.ts b/Engine/Extras/RoboMerge/v3/functional_tests/framework/src/framework.ts
--- a/Engine/Extras/RoboMerge/v3/functional_tests/framework/src/framework.ts
+++ b/Engine/Extras/RoboMerge/v3/functional_tests/framework/src/framework.ts
@@ -71,8 +71,8 @@ export class P4Client {
 		return this.p4.add(this.user, this.workspace, Path.join(this.root, file), binary)
 	}
 
-	edit(file: string) {
-		return this.p4.edit(this.user, this.workspace, Path.join(this.root, file))
+	edit(file: string, opts?: string[]) {
+		return this.p4.edit(this.user, this.workspace, Path.join(this.root, file), opts)
 	}
 
 	delete(file: string) {
@@ -147,7 +147,6 @@ type RobomergeBranchOptions = {
 	////////////////////
 	// NodeOptionFields
 	disabled: boolean
-	integrationMethod: string
 	forceAll: boolean
 	visibility: string[] | string
 	blockAssetFlow: string[]
@@ -189,6 +188,7 @@ type EdgeOptionFields = {
 	additionalSlackChannel: string
 	initialCL: number
 	p4MaxRowsOverride: number // use with care and check with p4 admins
+	integrationMethod: string
 
 	disallowSkip: boolean
 	incognitoMode: boolean
diff --git a/Engine/Extras/RoboMerge/v3/functional_tests/framework/src/test-perforce.ts b/Engine/Extras/RoboMerge/v3/functional_tests/framework/src/test-perforce.ts
--- a/Engine/Extras/RoboMerge/v3/functional_tests/framework/src/test-perforce.ts
+++ b/Engine/Extras/RoboMerge/v3/functional_tests/framework/src/test-perforce.ts
@@ -44,7 +44,9 @@ export class Perforce {
 		return this.exec(args)
 	}
 
-	edit(user: string, workspace: string, path: string) { return this.exec(['-u', user, '-c', workspace, 'edit', path]) }
+	edit(user: string, workspace: string, path: string, opts?: string[]) { 
+		return this.exec(['-u', user, '-c', workspace, 'edit', ...(opts || []), path]) 
+	}
 	delete(user: string, workspace: string, path: string) { return this.exec(['-u', user, '-c', workspace, 'delete', path]) }
 	sync(user: string, workspace: string) { return this.exec(['-u', user, '-c', workspace, 'sync']) }
 	populate(stream: string, description: string) { return this.exec(['populate', '-S', stream, '-r', '-d', description]) }
diff --git a/Engine/Extras/RoboMerge/v3/functional_tests/framework/src/tests.ts b/Engine/Extras/RoboMerge/v3/functional_tests/framework/src/tests.ts
--- a/Engine/Extras/RoboMerge/v3/functional_tests/framework/src/tests.ts
+++ b/Engine/Extras/RoboMerge/v3/functional_tests/framework/src/tests.ts
@@ -36,11 +36,12 @@ import { RequestShelf } from './tests/request-shelf'
 import { RequestShelfIndirectTarget } from './tests/request-shelf-indirect-target'
 import { ResolveAfterSkip } from './tests/resolve-after-skip'
 import { RespectStreamPath } from './tests/respect-stream-path'
-import { SyntaxErrorOnUnknownBranch } from './tests/syntax-error-on-unknown-branch'
+import { SkipIfUnchanged } from './tests/skip-if-unchanged'
 import { StompForwardingCommands } from './tests/stomp-forwarding-commands'
 import { StompWithAdd } from './tests/stomp-with-add'
 import { StompWithBranch } from './tests/stomp-with-branch'
 import { StompWithDelete } from './tests/stomp-with-delete'
+import { SyntaxErrorOnUnknownBranch } from './tests/syntax-error-on-unknown-branch'
 import { TestChain } from './tests/test-chain'
 import { TestEdgeGate } from './tests/test-edge-gate'
 import { TestFlags } from './tests/test-flags'
@@ -223,7 +224,9 @@ async function go() {
 		new BranchspecTest(p4),
 		new ExclusiveCheckoutUnlock(p4),
 		new ExcludeDescriptions(p4),
-		new ExcludeDescriptionsPerEdge(p4) // 50
+		new ExcludeDescriptionsPerEdge(p4), // 50
+
+		new SkipIfUnchanged(p4)
 	]
 
 	const TARGET_TEST_DEFS: [string[], string | null][] = [
diff --git a/Engine/Extras/RoboMerge/v3/functional_tests/framework/src/tests/edigrate-main-rev-to-release.ts b/Engine/Extras/RoboMerge/v3/functional_tests/framework/src/tests/edigrate-main-rev-to-release.ts
deleted file mode 100644
diff --git a/Engine/Extras/RoboMerge/v3/functional_tests/framework/src/tests/skip-if-unchanged.ts b/Engine/Extras/RoboMerge/v3/functional_tests/framework/src/tests/skip-if-unchanged.ts
new file mode 100644
--- /dev/null
+++ b/Engine/Extras/RoboMerge/v3/functional_tests/framework/src/tests/skip-if-unchanged.ts
@@ -0,0 +1,66 @@
+// Copyright Epic Games, Inc. All Rights Reserved.
+import { EdgeProperties, FunctionalTest, P4Util, Stream } from '../framework'
+
+const streams: Stream[] = [
+    {name: 'Main', streamType: 'mainline'},
+    {name: 'Dev-Pootle', streamType: 'development', parent: 'Main'},
+    {name: 'Dev-Perkin', streamType: 'development', parent: 'Main'}
+]
+
+export class SkipIfUnchanged extends FunctionalTest {
+
+	async setup() {
+        await this.p4.depot('stream', this.depotSpec())
+        await this.createStreamsAndWorkspaces(streams)
+
+		const mainClient = this.getClient('Main')
+        await P4Util.addFileAndSubmit(mainClient, 'unchanged.txt', 'Initial content', true)
+        await P4Util.addFileAndSubmit(mainClient, 'changed.txt', 'Initial content', true)
+        await P4Util.addFileAndSubmit(mainClient, 'changed_filetype.txt', 'Initial content', true)
+
+        const desc = 'Initial branch of files from Main'
+        await Promise.all([
+            this.p4.populate(this.getStreamPath('Dev-Perkin'), desc),
+            this.p4.populate(this.getStreamPath('Dev-Pootle'), desc)
+        ])
+	}
+
+	async run() {
+		const mainClient = this.getClient('Main')
+		await P4Util.editFileAndSubmit(mainClient, 'unchanged.txt', 'Initial content')
+		await P4Util.editFileAndSubmit(mainClient, 'changed.txt', 'Changed content')
+		await mainClient.edit('changed_filetype.txt', ['-t','text+w'])
+			.then(() => mainClient.submit("Changed filetype"))
+	}
+
+	getBranches() {
+		const mainDef = this.makeForceAllBranchDef('Main', ['Dev-Perkin', 'Dev-Pootle'])
+		mainDef.blockAssetFlow = [this.fullBranchName('Dev-Perkin')]
+		return [
+			mainDef,
+			this.makeForceAllBranchDef('Dev-Perkin', []),
+			this.makeForceAllBranchDef('Dev-Pootle', [])
+		]
+	}
+
+    getEdges() : EdgeProperties[] {
+        return [{ 
+            from: this.fullBranchName('Main'),
+            to: this.fullBranchName('Dev-Perkin'),
+            integrationMethod: "skip-if-unchanged"
+          }]
+    }
+
+    verify() {
+		return Promise.all([
+			this.checkHeadRevision('Dev-Perkin', 'unchanged.txt', 1),
+			this.checkHeadRevision('Dev-Pootle', 'unchanged.txt', 2),
+			this.checkHeadRevision('Dev-Perkin', 'changed.txt', 2),
+			this.checkHeadRevision('Dev-Pootle', 'changed.txt', 2),
+			this.checkHeadRevision('Dev-Perkin', 'changed_filetype.txt', 2),
+			this.checkHeadRevision('Dev-Pootle', 'changed_filetype.txt', 2),
+			this.ensureNotBlocked('Main', 'Dev-Perkin')
+		])
+	}
+}
+
diff --git a/Engine/Extras/RoboMerge/v3/public/js/flow.graph.js b/Engine/Extras/RoboMerge/v3/public/js/flow.graph.js
--- a/Engine/Extras/RoboMerge/v3/public/js/flow.graph.js
+++ b/Engine/Extras/RoboMerge/v3/public/js/flow.graph.js
@@ -201,8 +201,7 @@ class Graph {
         const isGated = new Map(srcBranchStatus.def.edgeProperties).get(dst).lastGoodCLPath
         this.connectedNodes.add(srcInfo.id);
         this.connectedNodes.add(dstInfo.id);
-        const edgeStyle = srcBranchStatus.def.convertIntegratesToEdits ? 'roboshelf' :
-            isForced ? (isGated ? 'gated' : 'forced') :
+        const edgeStyle = isForced ? (isGated ? 'gated' : 'forced') :
                 hasEdge(srcInfo.defaultDests, dstInfo) ? 'defaultFlow' :
                     hasEdge(srcInfo.blockAssetDests, dstInfo) ? 'blockAssets' : 'onRequest';
         const styles = [...EDGE_STYLES[edgeStyle]];
diff --git a/Engine/Extras/RoboMerge/v3/src/common/perforce.ts b/Engine/Extras/RoboMerge/v3/src/common/perforce.ts
--- a/Engine/Extras/RoboMerge/v3/src/common/perforce.ts
+++ b/Engine/Extras/RoboMerge/v3/src/common/perforce.ts
@@ -1814,6 +1814,11 @@ export class PerforceContext {
 		}
 	}
 
+	async hasDiff(roboWorkspace: RoboWorkspace, depotPath: string) {
+		const result = await this._execP4(roboWorkspace, ['diff', '-sr', depotPath], {quiet:true})
+		return result.length == 0
+	}
+
 	async fstat(roboWorkspace: RoboWorkspace, depotPath: string, quiet?: boolean) {
 		try {
 			return await this.execAndParse(roboWorkspace, ['fstat', depotPath], {quiet})
diff --git a/Engine/Extras/RoboMerge/v3/src/robo/branch-interfaces.ts b/Engine/Extras/RoboMerge/v3/src/robo/branch-interfaces.ts
--- a/Engine/Extras/RoboMerge/v3/src/robo/branch-interfaces.ts
+++ b/Engine/Extras/RoboMerge/v3/src/robo/branch-interfaces.ts
@@ -105,7 +105,6 @@ export type EditableBranch = BranchBase & {
 	reachable?: string[]
 	forcedDownstream?: string[]
 	enabled: boolean
-	convertIntegratesToEdits: boolean
 	visibility: string[] | string
 	blockAssetTargets: Set<string>
 	allowDeadend: boolean
diff --git a/Engine/Extras/RoboMerge/v3/src/robo/branchdefs.ts b/Engine/Extras/RoboMerge/v3/src/robo/branchdefs.ts
--- a/Engine/Extras/RoboMerge/v3/src/robo/branchdefs.ts
+++ b/Engine/Extras/RoboMerge/v3/src/robo/branchdefs.ts
@@ -9,7 +9,6 @@ const RESERVED_BRANCH_NAMES = ['NONE', 'DEFAULT', 'IGNORE', 'DEADEND', ''];
 
 export interface BotConfig {
 	defaultStreamDepot: string | null
-	defaultIntegrationMethod: string | null
 	isDefaultBot: boolean
 	noStreamAliases: boolean
 	globalNotify: string[]
@@ -62,7 +61,7 @@ export type BranchBase = typeof branchBasePrototype
 
 export class IntegrationMethod {
 	static NORMAL = 'normal'
-	static CONVERT_TO_EDIT = 'convert-to-edit'
+	static SKIP_IF_UNCHANGED = 'skip-if-unchanged'
 
 	static all() {
 		const cls: any = IntegrationMethod
@@ -113,7 +112,6 @@ const nodeOptionFieldsPrototype = {
 	...commonOptionFieldsPrototype,
 
 	disabled: false,
-	integrationMethod: '',
 	forceAll: false,
 	visibility: '' as string[] | string,
 	blockAssetFlow: [''],
@@ -156,6 +154,7 @@ const edgeOptionFieldsPrototype = {
 	nagWhenBlocked: true,
 
 	terminal: false, // changes go along terminal edges but no further
+	integrationMethod: 'normal',
 
 	implicitCommands: [''],
 
@@ -263,16 +262,15 @@ export class BranchDefs {
 		return undefined
 	}
 
-	private static checkValidIntegrationMethod(outErrors: string[], method: string, branchName: string) {
+	private static checkValidIntegrationMethod(outErrors: string[], method: string, name: string) {
 		if (IntegrationMethod.all().indexOf(method.toLowerCase()) === -1) {
-			outErrors.push(`Unknown integrationMethod '${method}' in '${branchName}'`)
+			outErrors.push(`Unknown integrationMethod '${method}' in '${name}'`)
 		}
 	}
 
 	static async parseAndValidate(logger: ContextualLogger, outErrors: string[], branchSpecsText: string, isPreviewing?: boolean): Promise<ParseResult> {
 		const defaultConfigForWholeBot: BotConfig = {
 			defaultStreamDepot: null,
-			defaultIntegrationMethod: null,
 			isDefaultBot: false,
 			noStreamAliases: false,
 			globalNotify: [],
@@ -343,10 +341,6 @@ export class BranchDefs {
 			}
 		}
 
-		if (defaultConfigForWholeBot.defaultIntegrationMethod) {
-			BranchDefs.checkValidIntegrationMethod(outErrors, defaultConfigForWholeBot.defaultIntegrationMethod, 'config')
-		}
-
 		const namesToIgnore = defaultConfigForWholeBot.branchNamesToIgnore.map(s => s.toUpperCase())
 		defaultConfigForWholeBot.branchNamesToIgnore = namesToIgnore
 
@@ -434,10 +428,6 @@ export class BranchDefs {
 				addAlias(upperName, def.streamName.toUpperCase())
 			}
 
-			if (def.integrationMethod) {
-				BranchDefs.checkValidIntegrationMethod(outErrors, def.integrationMethod!, def.name!)
-			}
-
 			// check all properties are known (could make things case insensitive here)
 			for (const keyName of Object.keys(def)) {
 				if (!nodeOptionFieldNames.has(keyName)) {
@@ -472,6 +462,10 @@ export class BranchDefs {
 					}
 				}
 
+				if (edge.integrationMethod) {
+					BranchDefs.checkValidIntegrationMethod(outErrors, edge.integrationMethod, `${edge.from}->${edge.to}`)
+				}
+
 				validateCommonOptions(edge)
 			}
 		}
diff --git a/Engine/Extras/RoboMerge/v3/src/robo/branchgraph.ts b/Engine/Extras/RoboMerge/v3/src/robo/branchgraph.ts
--- a/Engine/Extras/RoboMerge/v3/src/robo/branchgraph.ts
+++ b/Engine/Extras/RoboMerge/v3/src/robo/branchgraph.ts
@@ -3,7 +3,7 @@
 import { setDefault } from '../common/helper';
 import { BranchSpec } from '../common/perforce';
 import { Branch, BranchGraphInterface, EditableBranch } from './branch-interfaces';
-import { BotConfig, BranchDefs, BranchGraphDefinition, EdgeOptions, EdgeProperties, IntegrationMethod } from './branchdefs';
+import { BotConfig, BranchDefs, BranchGraphDefinition, EdgeOptions, EdgeProperties } from './branchdefs';
 import { NodeOptions, calculateStream } from './branchdefs';
 
 type ConfigBlendMode = 'override' | 'accumulate'
@@ -248,7 +248,6 @@ if (botname === '__TEST__') {
 			nagAcknowledgedSchedule: options.nagAcknowledgedSchedule || null,
 			nagAcknowledgedLeeway: options.nagAcknowledgedLeeway || null,
 			nagWhenBlocked: options.nagWhenBlocked || null,
-			convertIntegratesToEdits: (options.integrationMethod || this.config.defaultIntegrationMethod) === IntegrationMethod.CONVERT_TO_EDIT,
 			visibility: options.visibility || this.config.visibility,
 			get isMonitored() { return !!(this.bot && this.bot.isRunning) }
 		}
diff --git a/Engine/Extras/RoboMerge/v3/src/robo/edgebot.ts b/Engine/Extras/RoboMerge/v3/src/robo/edgebot.ts
--- a/Engine/Extras/RoboMerge/v3/src/robo/edgebot.ts
+++ b/Engine/Extras/RoboMerge/v3/src/robo/edgebot.ts
@@ -9,7 +9,7 @@ import { Change, coercePerforceWorkspace, ConflictedResolveNFile, EdgeServer, Ed
 import { VersionReader } from "../common/version";
 import { EdgeBotInterface, IPCControls, ReconsiderArgs } from "./bot-interfaces";
 import { AlreadyIntegrated, Branch, ChangeInfo, ConflictingFile, ExclusiveFile, Failure, MergeAction, PendingChange } from "./branch-interfaces";
-import { EdgeOptions } from "./branchdefs";
+import { EdgeOptions, IntegrationMethod } from "./branchdefs";
 import { PersistentConflict } from "./conflict-interfaces";
 import { BotEventTriggers } from "./events";
 import { Gate } from "./gate";
@@ -317,7 +317,7 @@ class EdgeBotImpl extends PerforceStatefulBot {
 		this.block(this.createEdgeBlockageInfo(failure, pending), pauseDurationSeconds)
 	}
 
-	performMerge(info: ChangeInfo, target: MergeAction, convertIntegratesToEdits: boolean, additionalDescription: string): Promise<EdgeIntegrationDetails> {
+	performMerge(info: ChangeInfo, target: MergeAction, additionalDescription: string): Promise<EdgeIntegrationDetails> {
 		// make sure we come back in here afterwords
 		this._log_action(`Merging CL ${info.cl} via ${this.fullName} (${target.mergeMode})`)
 
@@ -439,11 +439,6 @@ class EdgeBotImpl extends PerforceStatefulBot {
 			}
 		}
 
-		if (convertIntegratesToEdits) {
-			// just for visibility - not used by RoboMerge
-			description += '#ROBOMERGE-EDIGRATE\n'
-		}
-
 		target.description = description + additionalDescription
 
 		if (target.branch.config.streamServer != this.sourceBranch.config.streamServer) {
@@ -538,6 +533,16 @@ class EdgeBotImpl extends PerforceStatefulBot {
 		return workspaceName
 	}
 
+	private async SkipIntegration(pending: PendingChange, edgeServerAddress: string|undefined, msg: string) {
+		const event: AlreadyIntegrated = {change: pending.change, action: pending.action}
+		this.sourceNode.onAlreadyIntegrated(event)
+
+		// integration not necessary
+		this.edgeBotLogger.info(msg)
+		await this.p4.deleteCl(pending.change.targetWorkspace, pending.newCl, edgeServerAddress)
+		return new EdgeIntegrationDetails('nothing to do', msg)		
+	}
+
 	private async integrate(info: ChangeInfo, target: MergeAction) : Promise<EdgeIntegrationDetails> {
 		const to_integrate = info.cl
 		this._log_action(`Integrating CL ${to_integrate} to ${this.targetBranch.name}`)
@@ -588,14 +593,11 @@ class EdgeBotImpl extends PerforceStatefulBot {
 
 			case 'already_integrated':
 			case 'no_files': {
-				const msg = `Change ${to_integrate} was not necessary in ${this.targetBranch.name}`
-				const event: AlreadyIntegrated = {change: info, action: target}
-				this.sourceNode.onAlreadyIntegrated(event)
-
-				// integration not necessary
-				this.edgeBotLogger.info(msg)
-				await this.p4.deleteCl(info.targetWorkspace, changenum, edgeServerAddress)
-				return new EdgeIntegrationDetails('nothing to do', msg)
+				return await this.SkipIntegration(
+					pending,
+					edgeServerAddress,
+					`Change ${to_integrate} was not necessary in ${this.targetBranch.name}`
+				)
 			}
 		}
 
@@ -729,14 +731,15 @@ class EdgeBotImpl extends PerforceStatefulBot {
 	async _resolveChangelist(pending: PendingChange, submitRetries  = 3) : Promise<EdgeIntegrationDetails> {
 		// do a resolve with P4
 		this._log_action(`Resolving CL ${pending.change.cl} against ${this.targetBranch.name}`)
-//		const detail: ResolveResultDetail = 'detailed'
+
+		const edgeServerAddress = pending.change.edgeServerToHostShelf && pending.change.edgeServerToHostShelf.address
 
 		const result = await this.p4.resolve(
 			pending.change.targetWorkspace,
 			pending.newCl,
 			pending.action.mergeMode,
 			false, // detail === 'quick',
-			pending.change.edgeServerToHostShelf && pending.change.edgeServerToHostShelf.address
+			edgeServerAddress
 			)
 
 		if (pending.action.flags.has('manual') || pending.change.forceCreateAShelf) {
@@ -765,6 +768,37 @@ class EdgeBotImpl extends PerforceStatefulBot {
 		}
 		else
 		{
+			if (this.options.integrationMethod === IntegrationMethod.SKIP_IF_UNCHANGED) {
+				const opened = await this.p4.opened(pending.change.targetWorkspace, pending.newCl)
+
+				const PARALLEL_CHECKS = 25
+				let foundDiff = opened.some(x => x.action !== 'edit' && x.action !== 'integrate')
+				if (!foundDiff) {
+					const checkUnchanged = async (startIndex: number) => {
+						for (let i=startIndex; i < opened.length && !foundDiff; i += PARALLEL_CHECKS) {
+							if (await this.p4.hasDiff(pending.change.targetWorkspace, opened[i].depotFile)) {
+								foundDiff = true
+							}
+							else if (!foundDiff) { // abort if another parallel check found a diff to avoid unnecessary work
+								const fstat = (await this.p4.fstat(pending.change.targetWorkspace, opened[i].depotFile, true))[0]
+								if (fstat.headType != fstat.type) {
+									foundDiff = true
+								}
+							}
+						}
+					}
+					await Promise.all(Array.from(Array(PARALLEL_CHECKS).keys()).map(n => checkUnchanged(n)))
+				}
+				
+				if (!foundDiff) {
+					return await this.SkipIntegration(
+						pending,
+						edgeServerAddress,
+						`Skipping ${pending.change.cl} as all files are unchanged.`
+					)
+				}
+			}
+
 			if (this.options.approval) {
 				const approval = this.options.approval
 
diff --git a/Engine/Extras/RoboMerge/v3/src/robo/nodebot.ts b/Engine/Extras/RoboMerge/v3/src/robo/nodebot.ts
--- a/Engine/Extras/RoboMerge/v3/src/robo/nodebot.ts
+++ b/Engine/Extras/RoboMerge/v3/src/robo/nodebot.ts
@@ -2085,7 +2085,7 @@ export class NodeBot extends PerforceStatefulBot implements NodeBotInterface {
 					.map(s => s + '\n')
 					.join('')
 
-			const integrationResult = await targetEdge.performMerge(info, target, this.branch.convertIntegratesToEdits, additionalDescription)
+			const integrationResult = await targetEdge.performMerge(info, target, additionalDescription)
 
 			// wake up target nodebot, so look for further merges or conflict resolutions
 			const targetNodebot = target.branch.bot as NodeBot
diff --git a/Engine/Extras/RoboMerge/v3/src/robo/status-types.ts b/Engine/Extras/RoboMerge/v3/src/robo/status-types.ts
--- a/Engine/Extras/RoboMerge/v3/src/robo/status-types.ts
+++ b/Engine/Extras/RoboMerge/v3/src/robo/status-types.ts
@@ -150,7 +150,6 @@ export type BranchDefForStatus = {
 	blockAssetTargets: string[]
 	bot: string
 	config: any //NodeOptions
-	convertIntegratesToEdits: boolean
 	defaultFlow: string[]
 	flowsTo: string[]
 	forceFlowTo: string[]
