Remove unused status type

diff --git a/Engine/Extras/RoboMerge/v3/src/robo/status-types.ts b/Engine/Extras/RoboMerge/v3/src/robo/status-types.ts
--- a/Engine/Extras/RoboMerge/v3/src/robo/status-types.ts
+++ b/Engine/Extras/RoboMerge/v3/src/robo/status-types.ts
@@ -25,7 +25,6 @@ export type FailureKind =
 	'Syntax error' |
 	'Disallowed files' |
 	'Approval required' |
-	'Conversion to edits failure' |
 	'Unit Test error'
 
 export enum AvailableTypeEnum {
