Nodes will now process edges at different CLs in parallel rather than waiting for edges with older CLs to process to catch up

diff --git a/Engine/Extras/RoboMerge/v3/src/robo/edgebot.ts b/Engine/Extras/RoboMerge/v3/src/robo/edgebot.ts
--- a/Engine/Extras/RoboMerge/v3/src/robo/edgebot.ts
+++ b/Engine/Extras/RoboMerge/v3/src/robo/edgebot.ts
@@ -30,7 +30,7 @@ const JIRA_REGEX = /^\s*#jira\s+(.*)/i
 
 type EdgeIntegrationResult = 'ok' | 'shelved' | 'nothing to do' | 'skipped' | 'error'
 
-class EdgeIntegrationDetails {
+export class EdgeIntegrationDetails {
 	result: EdgeIntegrationResult
 	message = ''
 	cl = -1 // either submitted or shelved cl (may need Change in shelved case, let's see)
diff --git a/Engine/Extras/RoboMerge/v3/src/robo/nodebot.ts b/Engine/Extras/RoboMerge/v3/src/robo/nodebot.ts
--- a/Engine/Extras/RoboMerge/v3/src/robo/nodebot.ts
+++ b/Engine/Extras/RoboMerge/v3/src/robo/nodebot.ts
@@ -12,7 +12,7 @@ import { BlockagePauseInfo, BranchStatus } from './status-types';
 import { AlreadyIntegrated, Blockage, Branch, BranchArg, BranchGraphInterface, ChangeInfo, EndIntegratingToGateEvent, Failure } from './branch-interfaces';
 import { MergeAction, OperationResult, PendingChange, resolveBranchArg, StompedRevision, StompVerification, StompVerificationFile, UnlockVerification }  from './branch-interfaces';
 import { Conflicts } from './conflicts';
-import { EdgeBot, EdgeMergeResults } from './edgebot';
+import { EdgeBot, EdgeIntegrationDetails, EdgeMergeResults } from './edgebot';
 import { BotEventTriggers } from './events';
 import { GraphBot } from './graphbot';
 import { getPingTarget, isUserSlackGroup, makeClLink, SlackMessages } from './notifications';
@@ -267,11 +267,19 @@ export class NodeBot extends PerforceStatefulBot implements NodeBotInterface {
 		let minCl = this.lastCl
 		// Gather the list of available edges for tick -- edges may change mid-tick but we will only allow this set to perform integrations
 		let availableEdges: Map<string, EdgeBot> = new Map<string, EdgeBot>()
+		let availableEdgesAtCL: Map<number, Map<string, EdgeBot>> = new Map<number, Map<string, EdgeBot>>()
+
+		// There must always be at least an empty availableEdges map at the minCl
+		availableEdgesAtCL.set(minCl, new Map<string, EdgeBot>())
 
 		this.conflicts.checkForResolvedConflicts(this.lastCl, this.edges)
 		for (const edgeBot of this.edges.values()) {
 			if (edgeBot.isAvailable) {
 				availableEdges.set(edgeBot.targetBranch.upperName, edgeBot)
+				let edgesAtCL = availableEdgesAtCL.get(edgeBot.lastCl) || new Map<string, EdgeBot>()
+				edgesAtCL.set(edgeBot.targetBranch.upperName, edgeBot)
+				availableEdgesAtCL.set(edgeBot.lastCl, edgesAtCL)
+
 				if (edgeBot.lastCl > 0 && (minCl < 0 || edgeBot.lastCl < minCl)) {
 					minCl = edgeBot.lastCl
 				}
@@ -341,7 +349,32 @@ export class NodeBot extends PerforceStatefulBot implements NodeBotInterface {
 			this.ticksSinceLastNewP4Commit = 0
 
 			this.headCL = changes[0].change
-			await this._processListOfChanges(availableEdges, changes)
+
+			// list of changes except reconsiders
+			const changesToProcess: Change[] = []
+
+			// process manual changes even if paused or above the lastGoodCL
+			for (const change of changes) {
+				if (change.isUserRequest) {
+					this.nodeBotLogger.debug(`Processing manual change CL ${change.change}`)
+					await this._processAndMergeCl(availableEdges, change, true)
+				}
+				else {
+					changesToProcess.push(change)
+				}
+			}
+
+			if (!this.isAvailable) {
+				this.nodeBotLogger.verbose(`Not processing changes due to unavailability.`)
+				this.nodeBotLogger.debug(`Blocked? ${this.pauseState.isBlocked()}: ${this.pauseState.blockagePauseInfo}`)
+				this.nodeBotLogger.debug(`Paused? ${this.pauseState.isManuallyPaused()}: ${this.pauseState.manualPauseInfo}`)
+				this.nodeBotLogger.debug(`Availablity? ${this.pauseState.isAvailable()}: ${this.pauseState.availableInfo}`)
+			}
+			else {
+				// make sure the list is sorted in ascending order
+				changesToProcess.sort((a, b) => a.change - b.change)
+				await Promise.all(Array.from(availableEdgesAtCL.entries()).map(async ([minCL,edgesAtCL]) => this._processListOfChanges(edgesAtCL, minCL, changesToProcess)));
+			}
 		}
 		return true
 	}
@@ -1526,36 +1559,16 @@ export class NodeBot extends PerforceStatefulBot implements NodeBotInterface {
 		return change.user
 	}
 
-	private async _processListOfChanges(availableEdges: Map<string, EdgeBot>, allChanges: Change[]) {
-		// list of changes except reconsiders
-		const changes: Change[] = []
-
-		// process manual changes even if paused or above the lastGoodCL
-		for (const change of allChanges) {
-			if (change.isUserRequest) {
-				this.nodeBotLogger.debug(`Processing manual change CL ${change.change}`)
-				await this._processAndMergeCl(availableEdges, change, true)
-			}
-			else {
-				changes.push(change)
-			}
-		}
-
-		if (!this.isAvailable) {
-			this.nodeBotLogger.verbose(`Not processing changes due to unavailability.`)
-			this.nodeBotLogger.debug(`Blocked? ${this.pauseState.isBlocked()}: ${this.pauseState.blockagePauseInfo}`)
-			this.nodeBotLogger.debug(`Paused? ${this.pauseState.isManuallyPaused()}: ${this.pauseState.manualPauseInfo}`)
-			this.nodeBotLogger.debug(`Availablity? ${this.pauseState.isAvailable()}: ${this.pauseState.availableInfo}`)
-			return
-		}
-
-		// make sure the list is sorted in ascending order
-		changes.sort((a, b) => a.change - b.change)
+	private async _processListOfChanges(availableEdges: Map<string, EdgeBot>, minCL: number, changes: Change[]) {
 
 		const startTime = Date.now()
 
 		for (let changeIndex = 0; changeIndex < changes.length; ++changeIndex) {
 			const change = changes[changeIndex]
+			
+			if (change.change < minCL) {
+				continue
+			}
 
 			change.user = await this.getChangeAuthor(change)
 			
@@ -2035,77 +2048,93 @@ export class NodeBot extends PerforceStatefulBot implements NodeBotInterface {
 		return this.branchGraph.getBranch(name) || null;
 	}
 
-	private async _mergeClViaEdges(availableEdges: Map<string, EdgeBot>, info: ChangeInfo, ignoreEdgePauseState: boolean) : Promise<EdgeMergeResults> {
-		const mergedChanges = new EdgeMergeResults()
-
-		for (const target of info.targets!) {
-			if (!this.hasEdge(target.branch)) {
-				const err = new Error(`No edge exists for target branch "${target.branch.name}"\nCurrent Edges: ${this.printEdges()}`)
-				this.nodeBotLogger.printException(err)
-				throw err
-			}
+	private async _mergeClViaEdge(target: MergeAction, availableEdges: Map<string, EdgeBot>, info: ChangeInfo, ignoreEdgePauseState: boolean, mergedChanges: EdgeMergeResults) {
+		if (!this.hasEdge(target.branch)) {
+			const err = new Error(`No edge exists for target branch "${target.branch.name}"\nCurrent Edges: ${this.printEdges()}`)
+			this.nodeBotLogger.printException(err)
+			return err
+		}
 
-			const targetEdge = availableEdges.get(target.branch.upperName)
+		const targetEdge = availableEdges.get(target.branch.upperName)
 
-			// If the edge requested is currently not available this tick, simply log that we skipped it and move on
-			if (!targetEdge) {
-				this.nodeBotLogger.verbose(`No available edges can serve CL ${info.cl} -> ${target.branch.name} this tick.`)
-				continue
-			}
+		// If the edge requested is currently not available this tick, simply log that we skipped it and move on
+		if (!targetEdge) {
+			this.nodeBotLogger.verbose(`No available edges can serve CL ${info.cl} -> ${target.branch.name} this tick.`)
+			return
+		}
 
-			// Gate problem: say gate is 15, last integrated on edge was 12, incoming is 17. Edge is available because 12 < 15
-			// Note: good thing about using availability (in theory) - will not fetch loads of revisions repeatedly
+		// Gate problem: say gate is 15, last integrated on edge was 12, incoming is 17. Edge is available because 12 < 15
+		// Note: good thing about using availability (in theory) - will not fetch loads of revisions repeatedly
 
-			// assuming ignoreEdgePauseState is false only in the regular flow, so we're good updating edge last CLs
-			if (!ignoreEdgePauseState) {
+		// assuming ignoreEdgePauseState is false only in the regular flow, so we're good updating edge last CLs
+		if (!ignoreEdgePauseState) {
 
-				// give edge a change to take its gate into account
-				targetEdge.preIntegrate(info.cl)
+			// give edge a chance to take its gate into account
+			targetEdge.preIntegrate(info.cl)
 
-				// If this edge is paused, skip it but we'll have to report it so the node doesn't advance its lastCl
-				if (!targetEdge.isAvailable) {
-					mergedChanges.markSkipped(target.branch)
-					availableEdges.delete(targetEdge.targetBranch.upperName)
-					continue
-				}
+			// If this edge is paused, skip it but we'll have to report it so the node doesn't advance its lastCl
+			if (!targetEdge.isAvailable) {
+				mergedChanges.markSkipped(target.branch)
+				availableEdges.delete(targetEdge.targetBranch.upperName)
+				return
 			}
+		}
 
-			// If this is not a manual change, ensure we don't tell an edge that has moved past this point to redo work
-			if (!info.userRequest && info.cl <= targetEdge.lastCl) {
-				this.nodeBotLogger.debug(`${info.cl} is behind ${targetEdge.displayName} lastCL: ${targetEdge.lastCl}`)
-				continue
-			}
+		// If this is not a manual change, ensure we don't tell an edge that has moved past this point to redo work
+		if (!info.userRequest && info.cl <= targetEdge.lastCl) {
+			this.nodeBotLogger.debug(`${info.cl} is behind ${targetEdge.displayName} lastCL: ${targetEdge.lastCl}`)
+			return
+		}
 
-			if (this.tickJournal) {
-				++this.tickJournal.merges
-			}
+		if (this.tickJournal) {
+			++this.tickJournal.merges
+		}
 
-			const additionalDescription = 
-				targetEdge.implicitCommands
-					.map(s => s + '\n')
-					.join('')
+		const additionalDescription = 
+		targetEdge.implicitCommands
+			.map(s => s + '\n')
+			.join('')
 
-			const integrationResult = await targetEdge.performMerge(info, target, additionalDescription)
+		const integrationResult = await targetEdge.performMerge(info, target, additionalDescription)
 
-			// wake up target nodebot, so look for further merges or conflict resolutions
-			const targetNodebot = target.branch.bot as NodeBot
-			if (targetNodebot) {
-				targetNodebot.ticksSinceLastNewP4Commit = 0
-				targetNodebot.skipTickCounter = 0
-			}
+		// wake up target nodebot, so look for further merges or conflict resolutions
+		const targetNodebot = target.branch.bot as NodeBot
+		if (targetNodebot) {
+			targetNodebot.ticksSinceLastNewP4Commit = 0
+			targetNodebot.skipTickCounter = 0
+		}
 
-			// If the target edge is now blocked after the merge attempt, remove it from future merges this tick
-			if (targetEdge.isBlocked) {
-				this.nodeBotLogger.info(`${targetEdge.displayName} is blocked after attempting to merge ${info.cl}, removing edge from merges this tick.`)
-				availableEdges.delete(targetEdge.targetBranch.upperName)
-			}
+		// If the target edge is now blocked after the merge attempt, remove it from future merges this tick
+		if (targetEdge.isBlocked) {
+			this.nodeBotLogger.info(`${targetEdge.displayName} is blocked after attempting to merge ${info.cl}, removing edge from merges this tick.`)
+			availableEdges.delete(targetEdge.targetBranch.upperName)
+		}
+
+		mergedChanges.addIntegrationResult(target, integrationResult)
 
-			mergedChanges.addIntegrationResult(target, integrationResult)
+		return integrationResult
+	}
 
-			// if integration did something, enough time has probably passed to make worth checking to see if any manual requests came in
-			if (integrationResult.result !== 'nothing to do') {
-				await this.postIntegrate()
+	private async _mergeClViaEdges(availableEdges: Map<string, EdgeBot>, info: ChangeInfo, ignoreEdgePauseState: boolean) : Promise<EdgeMergeResults> {
+
+		const mergedChanges = new EdgeMergeResults()
+		const mergeResults = await Promise.all(info.targets!.map(async (target) => this._mergeClViaEdge(target, availableEdges, {...info}, ignoreEdgePauseState, mergedChanges)))
+
+		let workDone = false
+		for (const mergeResult of mergeResults) {
+			if (mergeResult instanceof Error) {
+				throw mergeResult
 			}
+			else if (mergeResult instanceof EdgeIntegrationDetails) {
+				if (mergeResult.result !== 'nothing to do') {
+					workDone = true
+				}
+			}
+		}
+
+		// if integration did something, enough time has probably passed to make worth checking to see if any manual requests came in
+		if (workDone) {
+			await this.postIntegrate()
 		}
 
 		return mergedChanges
