Catch login refresh errors and don't output things like ... THE PASSWORD ... in to the log

diff --git a/Engine/Extras/RoboMerge/v3/src/common/perforce.ts b/Engine/Extras/RoboMerge/v3/src/common/perforce.ts
--- a/Engine/Extras/RoboMerge/v3/src/common/perforce.ts
+++ b/Engine/Extras/RoboMerge/v3/src/common/perforce.ts
@@ -545,6 +545,15 @@ export class PerforceContext {
 							logger.warn(`Failed to refresh credentials for ${serverID}: ${serverAddress}. Will try again in 1 hour.`)
 						}
 					})
+					.catch(reason => {
+						if (!isExecP4Error(reason)) {
+							throw reason
+						}
+
+						const [_, output] = reason
+						
+						logger.warn(`Failed to refresh credentials for ${serverID}: ${serverAddress}.\nError: ${output}\n\nWill try again in 1 hour.`)
+					})
 			
 				}, 60*60*1000)
 			}
