Don't add #ROBOMERGE-COMMAND if incognito mode

diff --git a/Engine/Extras/RoboMerge/v3/src/robo/edgebot.ts b/Engine/Extras/RoboMerge/v3/src/robo/edgebot.ts
--- a/Engine/Extras/RoboMerge/v3/src/robo/edgebot.ts
+++ b/Engine/Extras/RoboMerge/v3/src/robo/edgebot.ts
@@ -374,12 +374,14 @@ class EdgeBotImpl extends PerforceStatefulBot {
 			description += '#robomerge #ignore\n'
 		}
 
-		if (info.overriddenCommand) {
-			// probably no need to remove #s, but feels safer
-			description += `#ROBOMERGE-COMMAND: ${info.overriddenCommand.replace(/#/g, '_')}\n`
-		}
-		if (info.macros.length > 0) {
-			description += `#ROBOMERGE-COMMAND: ${info.macros.join(', ')}\n`	
+		if (!this.incognitoMode) {
+			if (info.overriddenCommand) {
+				// probably no need to remove #s, but feels safer
+				description += `#ROBOMERGE-COMMAND: ${info.overriddenCommand.replace(/#/g, '_')}\n`
+			}
+			if (info.macros.length > 0) {
+				description += `#ROBOMERGE-COMMAND: ${info.macros.join(', ')}\n`	
+			}
 		}
 		let source = info.source
 		if (this.incognitoMode) {
