Fix unlock files failing when there are > 100 files to unlock

diff --git a/Engine/Extras/RoboMerge/v3/functional_tests/framework/src/tests.ts b/Engine/Extras/RoboMerge/v3/functional_tests/framework/src/tests.ts
--- a/Engine/Extras/RoboMerge/v3/functional_tests/framework/src/tests.ts
+++ b/Engine/Extras/RoboMerge/v3/functional_tests/framework/src/tests.ts
@@ -15,6 +15,7 @@ import { EdgeIndependence } from './tests/edge-independence'
 import { EdgeInitialCl } from './tests/edge-initial-cl'
 import { ExclusiveCheckout } from './tests/exclusive-checkout'
 import { ExclusiveCheckoutUnlock } from './tests/exclusive-checkout-unlock'
+import { ExclusiveCheckoutUnlockMultiple } from './tests/exclusive-checkout-unlock-multiple'
 import { ExcludeAuthors } from './tests/exclude-authors'
 import { ExcludeAuthorsPerEdge } from './tests/exclude-authors-per-edge'
 import { ExcludeDescriptions } from './tests/exclude-descriptions'
@@ -223,9 +224,10 @@ async function go() {
 		new PostToAdditionalChannel(p4),
 		new BranchspecTest(p4),
 		new ExclusiveCheckoutUnlock(p4),
-		new ExcludeDescriptions(p4),
-		new ExcludeDescriptionsPerEdge(p4), // 50
+		new ExclusiveCheckoutUnlockMultiple(p4),
+		new ExcludeDescriptions(p4), // 50
 
+		new ExcludeDescriptionsPerEdge(p4),
 		new SkipIfUnchanged(p4)
 	]
 
diff --git a/Engine/Extras/RoboMerge/v3/functional_tests/framework/src/tests/exclusive-checkout-unlock-multiple.ts b/Engine/Extras/RoboMerge/v3/functional_tests/framework/src/tests/exclusive-checkout-unlock-multiple.ts
new file mode 100644
--- /dev/null
+++ b/Engine/Extras/RoboMerge/v3/functional_tests/framework/src/tests/exclusive-checkout-unlock-multiple.ts
@@ -0,0 +1,49 @@
+// Copyright Epic Games, Inc. All Rights Reserved.
+import { FunctionalTest, P4Client, P4Util, RobomergeBranchSpec, Stream } from '../framework'
+
+const streams: Stream[] = [
+	{name: 'Main', streamType: 'mainline'},
+	{name: 'Release', streamType: 'release', parent: 'Main'}
+]
+
+export class ExclusiveCheckoutUnlockMultiple extends FunctionalTest {
+	mainSpec: RobomergeBranchSpec
+	releaseUser2Client: P4Client
+
+	async setup() {
+		await this.p4.depot('stream', this.depotSpec())
+		await this.createStreamsAndWorkspaces(streams, undefined, ['testuser1', 'testuser2'])
+
+		const mainClient = this.getClient('Main', 'testuser1')
+		await P4Util.addFileAndSubmit(mainClient, 'test.uasset', 'Dummy content', true)
+        await P4Util.addFileAndSubmit(mainClient, 'test2.uasset', 'Dummy content', true)
+		
+		this.releaseUser2Client = this.getClient('Release', 'testuser2')
+		await Promise.all([
+			mainClient.edit('test.uasset'),
+			mainClient.edit('test2.uasset'),
+			this.p4.populate(this.getStreamPath('Release'), 'Initial branch of files from Main')
+				.then(() => this.releaseUser2Client.sync())
+		])
+	}
+
+	async run() {
+		// second user commits change in Release while first user has file in main checked out
+		await P4Util.editFile(this.releaseUser2Client, 'test.uasset', 'New content')
+        await P4Util.editFile(this.releaseUser2Client, 'test2.uasset', 'New content')
+        await P4Util.submit(this.releaseUser2Client, "Files")
+	}
+
+	async verify() {
+		await this.ensureBlocked('Release', 'Main')
+        await this.verifyAndPerformUnlock('Release', 'Main')
+        await this.ensureNotBlocked('Release', 'Main')
+	}
+
+	getBranches() {
+		return [
+			this.makeForceAllBranchDef('Main', []),
+			this.makeForceAllBranchDef('Release', ['Main'])
+		]
+	}
+}
diff --git a/Engine/Extras/RoboMerge/v3/functional_tests/framework/src/tests/exclusive-checkout-unlock.ts b/Engine/Extras/RoboMerge/v3/functional_tests/framework/src/tests/exclusive-checkout-unlock.ts
--- a/Engine/Extras/RoboMerge/v3/functional_tests/framework/src/tests/exclusive-checkout-unlock.ts
+++ b/Engine/Extras/RoboMerge/v3/functional_tests/framework/src/tests/exclusive-checkout-unlock.ts
@@ -26,7 +26,7 @@ export class ExclusiveCheckoutUnlock extends FunctionalTest {
 	}
 
 	run() {
-		// second user commits change in Release well first user has file in main checked out
+		// second user commits change in Release while first user has file in main checked out
 		return P4Util.editFileAndSubmit(this.releaseUser2Client, 'test.uasset', 'New content')
 	}
 
@@ -34,7 +34,6 @@ export class ExclusiveCheckoutUnlock extends FunctionalTest {
 		await this.ensureBlocked('Release', 'Main')
         await this.verifyAndPerformUnlock('Release', 'Main')
         await this.ensureNotBlocked('Release', 'Main')
-
 	}
 
 	getBranches() {
diff --git a/Engine/Extras/RoboMerge/v3/run-tests.sh b/Engine/Extras/RoboMerge/v3/run-tests.sh
--- a/Engine/Extras/RoboMerge/v3/run-tests.sh
+++ b/Engine/Extras/RoboMerge/v3/run-tests.sh
@@ -93,6 +93,7 @@ docker run -d -p 8877:8877 -p 8811:8811 -p 25:25 -h robomerge_functtest --name r
     -e BOTNAME=ft1,ft2,ft3,ft4,targets  \
     -e ROBO_LOG_LEVEL=info \
     -e ROBO_SLACK_DOMAIN=http://localhost:8811 \
+    -e ROBO_EXCLUSIVE_LOCK_OPENEDS_TO_RUN=1 \
     robomerge node --trace-warnings dist/robo/watchdog.js
 
 docker run -t --hostname robomerge_functionaltests --name robomerge_functionaltests --network robomerge_functtest_network  \
diff --git a/Engine/Extras/RoboMerge/v3/src/common/perforce.ts b/Engine/Extras/RoboMerge/v3/src/common/perforce.ts
--- a/Engine/Extras/RoboMerge/v3/src/common/perforce.ts
+++ b/Engine/Extras/RoboMerge/v3/src/common/perforce.ts
@@ -77,6 +77,13 @@ export interface OpenedFileRecord {
 	movedFile?: string
 }
 
+export interface ExclusiveFileDetails {
+	depotPath: string
+	name: string
+	user: string
+	client: string
+}
+
 interface DescribeEntry {
 	depotFile: string
 	action: string
@@ -1478,6 +1485,32 @@ export class PerforceContext {
 		return this._execP4(roboWorkspace, ["change", "-d", changelist.toString()], {serverAddress: edgeServerAddress});
 	}
 
+	async getOpenedDetails(paths: {depotPath: string, name: string}[]) {
+		const results: ExclusiveFileDetails[] = []
+		const openedRequests: [{depotPath: string, name: string}, Promise<OpenedFileRecord[]>, boolean][] = []
+		for (const path of paths) {
+			openedRequests.push([path, this.opened(null, path.depotPath, true), true])
+		}
+
+		const failedResults: ExclusiveFileDetails[] = []
+		for (const [path, request, exclusive] of openedRequests) {
+			const recs = await request
+			if (recs.length > 0) {
+				// should only be one, since we're looking for exclusive check-out errors
+				results.push({depotPath: path.depotPath, name: path.name, user: recs[0].user, client: recs[0].client})
+			} else if (exclusive) {
+				// if we failed to find it as an exclusive check-out, try as non-exclusive which will find adds
+				openedRequests.push([path, this.opened(null, path.depotPath), false])
+			} else {
+				this.logger.warn(`Failed to get the opened status of ${path.depotPath}`)
+				failedResults.push({depotPath: path.depotPath, name: path.name, user: "", client: ""})
+			}
+		}
+		results.push(...failedResults)
+
+		return results
+	}
+
 	// run p4 'opened' command: lists files in changelist with details of edit state (e.g. if a copy, provides source path)
 	opened(roboWorkspace: RoboWorkspace | null, arg: number | string, exclusive?: boolean) {
 		const args = ['opened']
diff --git a/Engine/Extras/RoboMerge/v3/src/robo/branch-interfaces.ts b/Engine/Extras/RoboMerge/v3/src/robo/branch-interfaces.ts
--- a/Engine/Extras/RoboMerge/v3/src/robo/branch-interfaces.ts
+++ b/Engine/Extras/RoboMerge/v3/src/robo/branch-interfaces.ts
@@ -1,6 +1,6 @@
 // Copyright Epic Games, Inc. All Rights Reserved.
 
-import { BranchSpec, ConflictedResolveNFile, EdgeServer } from '../common/perforce';
+import { BranchSpec, ConflictedResolveNFile, EdgeServer, ExclusiveFileDetails } from '../common/perforce';
 // TODO: Remove Circular Dependency on bot-interfaces
 import { NodeBotInterface } from './bot-interfaces';
 import { FailureKind } from './status-types';
@@ -56,13 +56,6 @@ export interface StompVerification extends OperationResult {
 	remainingAllBinary?: boolean // Check to see if unresolved non-binary files are remaining -- we shouldn't stomp those!
 }
 
-export interface ExclusiveFile {
-	depotPath: string
-	name: string
-	user: string
-	client: string
-}
-
 export interface ExclusiveLockUser {
     user: string
     userEmail: Promise<string | null>
@@ -70,12 +63,12 @@ export interface ExclusiveLockUser {
 
 export interface ExclusiveLockInfo {
 	exclusiveLockUsers: ExclusiveLockUser[]
-	exclusiveFiles: ExclusiveFile[]
+	exclusiveFiles: ExclusiveFileDetails[]
 }
 
 export interface UnlockVerification extends OperationResult {
 	validRequest?: boolean
-	lockedFiles?: ExclusiveFile[] // Array of locked files
+	lockedFiles?: ExclusiveFileDetails[] // Array of locked files
 }
 
 export interface BranchGraphInterface {
diff --git a/Engine/Extras/RoboMerge/v3/src/robo/edgebot.ts b/Engine/Extras/RoboMerge/v3/src/robo/edgebot.ts
--- a/Engine/Extras/RoboMerge/v3/src/robo/edgebot.ts
+++ b/Engine/Extras/RoboMerge/v3/src/robo/edgebot.ts
@@ -5,10 +5,10 @@ import * as p4util from '../common/p4util';
 
 import { ContextualLogger } from "../common/logger";
 import { Recipients } from "../common/mailer";
-import { Change, coercePerforceWorkspace, ConflictedResolveNFile, EdgeServer, EditChangeOpts, EXCLUSIVE_CHECKOUT_REGEX, getRootDirectoryForBranch, IntegrationSource, IntegrationTarget, isExecP4Error, OpenedFileRecord, PerforceContext } from "../common/perforce";
+import { Change, coercePerforceWorkspace, ConflictedResolveNFile, EdgeServer, EditChangeOpts, ExclusiveFileDetails, EXCLUSIVE_CHECKOUT_REGEX, getRootDirectoryForBranch, IntegrationSource, IntegrationTarget, isExecP4Error, PerforceContext } from "../common/perforce";
 import { VersionReader } from "../common/version";
 import { EdgeBotInterface, IPCControls, ReconsiderArgs } from "./bot-interfaces";
-import { AlreadyIntegrated, Branch, ChangeInfo, ConflictingFile, ExclusiveFile, Failure, MergeAction, PendingChange } from "./branch-interfaces";
+import { AlreadyIntegrated, Branch, ChangeInfo, ConflictingFile, Failure, MergeAction, PendingChange } from "./branch-interfaces";
 import { EdgeOptions, IntegrationMethod } from "./branchdefs";
 import { PersistentConflict } from "./conflict-interfaces";
 import { BotEventTriggers } from "./events";
@@ -16,6 +16,7 @@ import { Gate } from "./gate";
 import { NodeBot } from "./nodebot";
 import { isUserAKnownBot, postMessageToChannel } from "./notifications";
 import { PerforceStatefulBot } from "./perforce-stateful-bot";
+import { getExclusiveLockOpenedsToRun } from "./roboargs";
 import { Context } from "./settings";
 import { SlackMessageStyles } from "./slack";
 import { PauseState } from "./state-interfaces";
@@ -23,7 +24,6 @@ import { BlockagePauseInfo, BlockagePauseInfoMinimal, EdgeStatusFields } from ".
 import { getIntegrationOwner } from "./targets";
 
 const FAILED_CHANGELIST_PAUSE_TIMEOUT_SECONDS = 15 * 60
-const MAX_OPENED_COMMANDS_TO_ISSUE = 100
 const JIRA_REGEX = /^\s*#jira\s+(.*)/i
 
 // type ResolveResultDetail = 'quick' | 'detailed'
@@ -258,30 +258,21 @@ class EdgeBotImpl extends PerforceStatefulBot {
 
 	private async analyzeIntegrationError(errors: string[]) {
 		
-		const results: ExclusiveFile[] = []
-		const openedRequests: [RegExpMatchArray, Promise<OpenedFileRecord[]>, boolean][] = []
+		let results: ExclusiveFileDetails[] = []
+		let paths = []
 		for (const err of errors) {
 			const match = err.match(EXCLUSIVE_CHECKOUT_REGEX)
 			if (match) {
-				if (openedRequests.length < MAX_OPENED_COMMANDS_TO_ISSUE) {
-					openedRequests.push([match, this.p4.opened(null, match[1] + match[2], true), true])
+				if (paths.length < getExclusiveLockOpenedsToRun()) {
+					paths.push({depotPath: match[1] + match[2], name: match[2]})
 				} else {
 					results.push({depotPath: match[1] + match[2], name: match[2], user: "", client: ""})
 				}
 			}
 		}
 
-		for (const [match, request, exclusive] of openedRequests) {
-			const recs = await request
-			if (recs.length > 0) {
-				// should only be one, since we're looking for exclusive check-out errors
-				results.push({depotPath: match[1] + match[2], name: match[2], user: recs[0].user, client: recs[0].client})
-			} else if (exclusive && openedRequests.length < MAX_OPENED_COMMANDS_TO_ISSUE) {
-				// if we failed to find it as an exclusive check-out, try as non-exclusive which will find adds
-				openedRequests.push([match, this.p4.opened(null, match[1] + match[2]), false])
-			} else  {
-				results.push({depotPath: match[1] + match[2], name: match[2], user: "", client: ""})
-			}
+		if (paths.length > 0) {
+			results = [...(await this.p4.getOpenedDetails(paths)),...results]
 		}
 
 		return results
diff --git a/Engine/Extras/RoboMerge/v3/src/robo/nodebot.ts b/Engine/Extras/RoboMerge/v3/src/robo/nodebot.ts
--- a/Engine/Extras/RoboMerge/v3/src/robo/nodebot.ts
+++ b/Engine/Extras/RoboMerge/v3/src/robo/nodebot.ts
@@ -5,7 +5,7 @@ import * as p4util from '../common/p4util';
 import { _nextTick } from '../common/helper';
 import { ContextualLogger } from '../common/logger';
 import { Mailer, MailParams, Recipients } from '../common/mailer';
-import { Change, ConflictedResolveNFile, PerforceContext, RoboWorkspace, coercePerforceWorkspace } from '../common/perforce';
+import { Change, ConflictedResolveNFile, ExclusiveFileDetails, PerforceContext, RoboWorkspace, coercePerforceWorkspace } from '../common/perforce';
 import { IPCControls, NodeBotInterface, QueuedChange, ReconsiderArgs } from './bot-interfaces';
 import { ApprovalOptions, BotConfig, EdgeOptions, NodeOptions } from './branchdefs'
 import { BlockagePauseInfo, BranchStatus } from './status-types';
@@ -17,6 +17,7 @@ import { BotEventTriggers } from './events';
 import { GraphBot } from './graphbot';
 import { getPingTarget, isUserSlackGroup, makeClLink, SlackMessages } from './notifications';
 import { PerforceStatefulBot } from './perforce-stateful-bot';
+import { getExclusiveLockOpenedsToRun } from './roboargs';
 import { BlockageNodeOpUrls, OperationUrlHelper } from './roboserver';
 import { Context } from './settings';
 import { SlackFile, SlackMessage, SlackMessageField, SlackMessageStyles } from './slack';
@@ -1074,10 +1075,22 @@ export class NodeBot extends PerforceStatefulBot implements NodeBotInterface {
 		}
 		if (conflict.kind !== 'Exclusive check-out') {
 			return { success: false, message: `Unlock only usable when edge is blocked by exclusive check-out ${targetBranch.name}. ` + 
-				`If you thinks this is an error, contact Robomerge support.` }
+				`If you think this is an error, contact Robomerge support.` }
 		}
 
-		return { success: true, validRequest: true, message: "Unlock verified", lockedFiles: blockageInfo.additionalInfo.exclusiveFiles}
+		// Verify that we have the client/user information for all locked files
+		let lockedFiles: ExclusiveFileDetails[] = blockageInfo.additionalInfo.exclusiveFiles
+		if (lockedFiles.length > getExclusiveLockOpenedsToRun()) {
+			lockedFiles = [...lockedFiles.slice(0,getExclusiveLockOpenedsToRun()), ...(await this.p4.getOpenedDetails(lockedFiles.slice(getExclusiveLockOpenedsToRun())))]
+		}
+
+		const unknownClients = lockedFiles.filter((file) => file.client.length == 0)
+		if (unknownClients.length > 0) {
+			this.logger.warn(`Unable to determine owning client for locked files:\n${unknownClients.map((file) => file.client).join("\n")}`)
+			return { success: false, message: "Unable to determine owning client for locked files" }
+		}
+
+		return { success: true, validRequest: true, message: "Unlock verified", lockedFiles}
 	}
 
 	// First, re-verify the unlock request
diff --git a/Engine/Extras/RoboMerge/v3/src/robo/notifications.ts b/Engine/Extras/RoboMerge/v3/src/robo/notifications.ts
--- a/Engine/Extras/RoboMerge/v3/src/robo/notifications.ts
+++ b/Engine/Extras/RoboMerge/v3/src/robo/notifications.ts
@@ -5,7 +5,8 @@ import { Args } from '../common/args';
 import { Badge } from '../common/badge';
 import { Random, setDefault } from '../common/helper';
 import { ContextualLogger } from '../common/logger';
-import { Blockage, Branch, BranchArg, ExclusiveFile, ExclusiveLockInfo, ForcedCl, MergeAction, NodeOpUrlGenerator, resolveBranchArg } from './branch-interfaces';
+import { ExclusiveFileDetails } from '../common/perforce';
+import { Blockage, Branch, BranchArg,  ExclusiveLockInfo, ForcedCl, MergeAction, NodeOpUrlGenerator, resolveBranchArg } from './branch-interfaces';
 import { PersistentConflict, Resolution } from './conflict-interfaces';
 import { BotEventHandler, BotEvents } from './events';
 import { NodeBot } from './nodebot';
@@ -598,7 +599,7 @@ export class BotNotifications implements BotEventHandler {
 			if (isNew && blockage.failure.kind === 'Exclusive check-out') {
 				const exclusiveLockInfo = blockage.failure.additionalInfo as ExclusiveLockInfo
 
-				let authorDict = new Map<string, ExclusiveFile[]>()
+				let authorDict = new Map<string, ExclusiveFileDetails[]>()
 				for (const file of exclusiveLockInfo.exclusiveFiles) {
 					setDefault(authorDict, file.user.toLowerCase(), []).push(file)
 				}
diff --git a/Engine/Extras/RoboMerge/v3/src/robo/robo.ts b/Engine/Extras/RoboMerge/v3/src/robo/robo.ts
--- a/Engine/Extras/RoboMerge/v3/src/robo/robo.ts
+++ b/Engine/Extras/RoboMerge/v3/src/robo/robo.ts
@@ -23,6 +23,7 @@ import { NodeBot } from './nodebot';
 import { roboAnalytics, setGlobalAnalytics } from './roboanalytics';
 import { RoboServer } from './roboserver';
 import { notificationsInit } from './notifications'
+import { roboArgsInit } from './roboargs'
 import { settingsInit } from './settings'
 import { Vault } from './vault'
 import * as Preview from './preview'
@@ -145,6 +146,15 @@ const COMMAND_LINE_ARGS: {[param: string]: Arg<any>} = {
 		dflt: false
 	},
 
+	exclusiveLockOpenedsToRun: {
+		match: /^-exclusiveLockOpenedsToRun=([0-9]+)$/,
+		env: 'ROBO_EXCLUSIVE_LOCK_OPENEDS_TO_RUN',
+		parse: (str: string) => {
+			return parseInt(str)
+		},
+		dflt: 100
+	},
+
 	// Sentry environment designation -- use 'PROD' to enable Sentry bug tracking
 	epicEnv: {
 		match: /^(-epicEnv)$/,
@@ -213,6 +223,7 @@ const args = maybeNullArgs!
 
 const vault = new Vault(args.vault)
 notificationsInit(args, vault)
+roboArgsInit(args)
 settingsInit(args)
 Preview.init(args.branchSpecsRootPath)
 
diff --git a/Engine/Extras/RoboMerge/v3/src/robo/roboargs.ts b/Engine/Extras/RoboMerge/v3/src/robo/roboargs.ts
new file mode 100644
--- /dev/null
+++ b/Engine/Extras/RoboMerge/v3/src/robo/roboargs.ts
@@ -0,0 +1,12 @@
+// Copyright Epic Games, Inc. All Rights Reserved.
+
+import { Args } from '../common/args';
+
+let args: Args
+export function roboArgsInit(inArgs: Args) {
+	args = inArgs
+}
+
+export function getExclusiveLockOpenedsToRun() {
+	return args.exclusiveLockOpenedsToRun
+} 
\ No newline at end of file
