Branchspec validation will now return and display both errors and warnings (when previewing)
Added validation for existance of branchspec
Added validation for existance of lastGoodCLPath

diff --git a/Engine/Extras/RoboMerge/v3/public/js/preview.js b/Engine/Extras/RoboMerge/v3/public/js/preview.js
--- a/Engine/Extras/RoboMerge/v3/public/js/preview.js
+++ b/Engine/Extras/RoboMerge/v3/public/js/preview.js
@@ -178,8 +178,7 @@ function doit(query)
 			populateControlPanel(bots ? bots[1].split(',') : [], Array.from(botNames));
 
 			const reload = !!location.search.match(/reload/i);
-			if (reload)
-			{
+			if (reload) {
 				document.getElementById('controlpanel').open = true;
 			}
 
@@ -187,9 +186,17 @@ function doit(query)
 				data.allBranches,
 				{ botNames: [], ...parseOptions(location.search)},
 				false)
-		);
+			);
 			$('.bots').html([...botNames].map(s => `<tt>${s.toLowerCase()}</tt>`).join(', '))
-			$('#success-panel').show();
+			if (data.validationWarnings && data.validationWarnings.length > 0) {
+				const $warningPanel = $('#warning-panel');
+				const warningMsg = "Warning:" + data.validationWarnings.join('\n').replace(/\t/g, '    ')
+				$('pre', $warningPanel).text(warningMsg);
+				$warningPanel.show();
+			}
+			else {
+				$('#success-panel').show();
+			}
 			$('#control-panel').show();
 			$('#controlpanel').show();
 		},
diff --git a/Engine/Extras/RoboMerge/v3/public/preview.html b/Engine/Extras/RoboMerge/v3/public/preview.html
--- a/Engine/Extras/RoboMerge/v3/public/preview.html
+++ b/Engine/Extras/RoboMerge/v3/public/preview.html
@@ -35,6 +35,10 @@
 				<div class="alert alert-success" role="alert">Graph for <span class="bots"></span> from <a href="https://p4-swarm.epicgames.net/changes/{{ cl }}" target="_blank">CL#{{ cl }}</a></div>
 			</div>
 			<div class="col">
+				<div id="warning-panel" style="display:none;"> <!-- show or hide in script -->
+					<div class="alert alert-warning" role="alert">Warnings while previewing <span class="bots"></span> from <a href="https://p4-swarm.epicgames.net/changes/{{ cl }}" target="_blank">CL#{{ cl }}</a></div>
+					<pre></pre>
+				</div>
 				<div id="error-panel" style="display:none;"> <!-- show or hide in script -->
 					<div class="alert alert-warning" role="alert">Errors attempting to preview <span class="bots"></span> from <a href="https://p4-swarm.epicgames.net/changes/{{ cl }}" target="_blank">CL#{{ cl }}</a></div>
 					<pre></pre>
diff --git a/Engine/Extras/RoboMerge/v3/src/common/perforce.ts b/Engine/Extras/RoboMerge/v3/src/common/perforce.ts
--- a/Engine/Extras/RoboMerge/v3/src/common/perforce.ts
+++ b/Engine/Extras/RoboMerge/v3/src/common/perforce.ts
@@ -741,12 +741,20 @@ export class PerforceContext {
 			}
 			return streamSpec
 		}
-		catch
-		{
+		catch {
 			return null
 		}
 	}
 
+	async branchExists(branchName: string) {
+		try {
+			return (await this.execAndParse(null, ['branches','-e',branchName])).length == 1
+		}
+		catch {
+			return false
+		}
+	}
+
 	// find a workspace for the given user
 	// output format is an array of workspace names
 	async find_workspaces(user?: string, options?: {edgeServerAddress?: string, includeUnloaded?: boolean}) {
diff --git a/Engine/Extras/RoboMerge/v3/src/robo/autobranchupdater.ts b/Engine/Extras/RoboMerge/v3/src/robo/autobranchupdater.ts
--- a/Engine/Extras/RoboMerge/v3/src/robo/autobranchupdater.ts
+++ b/Engine/Extras/RoboMerge/v3/src/robo/autobranchupdater.ts
@@ -148,12 +148,11 @@ export class AutoBranchUpdater implements Bot {
 			return
 		}
 
-		const validationErrors: string[] = []
-		const result = await BranchDefs.parseAndValidate(this.abuLogger, validationErrors, branchGraphText)
+		const result = await BranchDefs.parseAndValidate(this.abuLogger, branchGraphText)
 		if (!result.branchGraphDef) {
 			// @todo email author of changes!
 			let errText = 'failed to parse/validate branch specs file\n'
-			for (const error of validationErrors) {
+			for (const error of result.validationErrors) {
 				errText += `${error}\n`
 			}
 			this.abuLogger.error(errText.trim())
diff --git a/Engine/Extras/RoboMerge/v3/src/robo/branchdefs.ts b/Engine/Extras/RoboMerge/v3/src/robo/branchdefs.ts
--- a/Engine/Extras/RoboMerge/v3/src/robo/branchdefs.ts
+++ b/Engine/Extras/RoboMerge/v3/src/robo/branchdefs.ts
@@ -199,7 +199,9 @@ export interface BranchGraphDefinition {
 // eventually, switch branch map to using a separate class defined here for the branch definitions
 
 
-function validateCommonOptions(options: Partial<CommonOptionFields>) {
+async function validateCommonOptions(logger: ContextualLogger, options: Partial<CommonOptionFields>, isPreviewing?: boolean) {
+	const errors: string[] = []
+	const warnings: string[] = []
 	if (options.integrationWindow) {
 		for (const pane of options.integrationWindow) {
 			if (pane.daysOfTheWeek) {
@@ -207,19 +209,32 @@ function validateCommonOptions(options: Partial<CommonOptionFields>) {
 					const dayStr = pane.daysOfTheWeek[index]
 					const day = dayStr.slice(0, 3).toLowerCase()
 					if (DAYS_OF_THE_WEEK.indexOf(day) < 0) {
-						throw new Error(`Unknown day of the week ${dayStr}`)
+						errors.push(`Unknown day of the week ${dayStr}`)
 					}
 					pane.daysOfTheWeek[index] = day as DayOfTheWeek
 				}
 			}
 		}
 	}
+	if (typeof(options.lastGoodCLPath) === 'string') {
+		const fstat = await PerforceContext.getServerContext(logger).fstat(null, options.lastGoodCLPath, true)
+		if (fstat.length == 0) {
+			if (isPreviewing) {
+				warnings.push(`Unable to find lastGoodCLPath ${options.lastGoodCLPath}`)
+			} else {
+				errors.push(`Unable to find lastGoodCLPath ${options.lastGoodCLPath}`)
+			}
+		}
+	}
+	return {errors, warnings}
 }
 
 
 interface ParseResult {
 	branchGraphDef: BranchGraphDefinition | null
 	config: BotConfig
+	validationErrors: string[],
+	validationWarnings: string[]
 }
 
 export type StreamResult = {
@@ -262,13 +277,13 @@ export class BranchDefs {
 		return undefined
 	}
 
-	private static checkValidIntegrationMethod(outErrors: string[], method: string, name: string) {
+	private static checkValidIntegrationMethod(validationErrors: string[], method: string, name: string) {
 		if (IntegrationMethod.all().indexOf(method.toLowerCase()) === -1) {
-			outErrors.push(`Unknown integrationMethod '${method}' in '${name}'`)
+			validationErrors.push(`Unknown integrationMethod '${method}' in '${name}'`)
 		}
 	}
 
-	static async parseAndValidate(logger: ContextualLogger, outErrors: string[], branchSpecsText: string, isPreviewing?: boolean): Promise<ParseResult> {
+	static async parseAndValidate(logger: ContextualLogger, branchSpecsText: string, isPreviewing?: boolean): Promise<ParseResult> {
 		const defaultConfigForWholeBot: BotConfig = {
 			defaultStreamDepot: null,
 			isDefaultBot: false,
@@ -294,6 +309,9 @@ export class BranchDefs {
 			badgeUrlOverride: ''
 		}
 
+		let validationErrors: string[] = []
+		let validationWarnings: string[] = []
+
 		let branchGraphRaw: any
 		try {
 			branchGraphRaw = jsonlint.parse(branchSpecsText)
@@ -303,8 +321,8 @@ export class BranchDefs {
 			}
 		}
 		catch (err) {
-			outErrors.push(err)
-			return {branchGraphDef: null, config: defaultConfigForWholeBot}
+			validationErrors.push(err)
+			return {branchGraphDef: null, config: defaultConfigForWholeBot, validationErrors, validationWarnings}
 		}
 
 		if (branchGraphRaw.alias) {
@@ -332,8 +350,8 @@ export class BranchDefs {
 						}
 					}
 					if (!macrosLower) {
-						outErrors.push(`Invalid macro property: '${value}'`)
-						return {branchGraphDef: null, config: defaultConfigForWholeBot}
+						validationErrors.push(`Invalid macro property: '${value}'`)
+						return {branchGraphDef: null, config: defaultConfigForWholeBot, validationErrors, validationWarnings}
 					}
 					value = macrosLower
 				}
@@ -353,7 +371,7 @@ export class BranchDefs {
 		// Check for duplicate branch names
 		for (const def of branchesFromJSON) {
 			if (!def.name) {
-				outErrors.push(`Unable to parse branch definition: ${JSON.stringify(def)}`)
+				validationErrors.push(`Unable to parse branch definition: ${JSON.stringify(def)}`)
 				continue
 			}
 
@@ -361,13 +379,13 @@ export class BranchDefs {
 
 			const nameError = BranchDefs.checkName(def.name)
 			if (nameError) {
-				outErrors.push(nameError)
+				validationErrors.push(nameError)
 				continue
 			}
 
 			const upperName = def.name.toUpperCase()
 			if (names.has(upperName)) {
-				outErrors.push(`Duplicate branch name '${upperName}'`)
+				validationErrors.push(`Duplicate branch name '${upperName}'`)
 			}
 			else {
 				names.set(upperName, upperName)
@@ -380,35 +398,39 @@ export class BranchDefs {
 				def.streamSubpath
 			)
 
-			if (!isPreviewing && streamResult.stream && !await PerforceContext.getServerContext(logger, def.streamServer).stream(streamResult.stream)) {
-				outErrors.push(`Stream ${streamResult.stream} not found`)
+			if (streamResult.stream && !await PerforceContext.getServerContext(logger, def.streamServer).stream(streamResult.stream)) {
+				if (isPreviewing) {
+					validationWarnings.push(`Stream ${streamResult.stream} not found`)
+				} else {
+					validationErrors.push(`Stream ${streamResult.stream} not found`)
+				}
 			}
 		}
 
 		// Check for duplicate aliases (and that branches/aliases are not in the ignore list)
 		const addAlias = (upperBranchName: string, upperAlias: string) => {
 			if (namesToIgnore.indexOf(upperBranchName) >= 0) {
-				outErrors.push(upperBranchName + ' branch is in branchNamesToIgnore')
+				validationErrors.push(upperBranchName + ' branch is in branchNamesToIgnore')
 			}
 
 			if (namesToIgnore.indexOf(upperAlias) >= 0) {
-				outErrors.push(upperAlias + ' alias is in branchNamesToIgnore')
+				validationErrors.push(upperAlias + ' alias is in branchNamesToIgnore')
 			}
 
 			if (!upperAlias) {
-				outErrors.push(`Empty alias for '${upperBranchName}'`)
+				validationErrors.push(`Empty alias for '${upperBranchName}'`)
 				return
 			}
 
 			const nameError = BranchDefs.checkName(upperAlias)
 			if (nameError) {
-				outErrors.push(nameError)
+				validationErrors.push(nameError)
 				return
 			}
 
 			const existing = names.get(upperAlias)
 			if (existing && existing !== upperBranchName) {
-				outErrors.push(`Duplicate alias '${upperAlias}' for '${existing}' and '${upperBranchName}'`)
+				validationErrors.push(`Duplicate alias '${upperAlias}' for '${existing}' and '${upperBranchName}'`)
 			}
 			else {
 				names.set(upperAlias, upperBranchName)
@@ -435,17 +457,19 @@ export class BranchDefs {
 				}
 			}
 
-			validateCommonOptions(def)
+			const results = await validateCommonOptions(logger, def, isPreviewing)
+			validationErrors.push(...results.errors)
+			validationWarnings.push(...results.warnings)
 		}
 
 		// Check edge properties
 		if (branchGraph.edges) {
 			for (const edge of branchGraph.edges) {
 				if (!names.get(edge.from.toUpperCase())) {
-					outErrors.push('Unrecognised source node in edge property ' + edge.from)
+					validationErrors.push('Unrecognised source node in edge property ' + edge.from)
 				}
 				if (!names.get(edge.to.toUpperCase())) {
-					outErrors.push('Unrecognised target node in edge property ' + edge.to)
+					validationErrors.push('Unrecognised target node in edge property ' + edge.to)
 				}
 
 				// check all properties are known (could make things case insensitive here)
@@ -463,11 +487,22 @@ export class BranchDefs {
 				}
 
 				if (edge.integrationMethod) {
-					BranchDefs.checkValidIntegrationMethod(outErrors, edge.integrationMethod, `${edge.from}->${edge.to}`)
+					BranchDefs.checkValidIntegrationMethod(validationErrors, edge.integrationMethod, `${edge.from}->${edge.to}`)
 				}
 
-				validateCommonOptions(edge)
-			}
+				if (edge.branchspec && !await PerforceContext.getServerContext(logger).branchExists(edge.branchspec)) {
+					if (isPreviewing) {
+						validationWarnings.push(`Branchspec ${edge.branchspec} does not exist.`)
+					} 
+					else {
+						validationErrors.push(`Branchspec ${edge.branchspec} does not exist.`)
+					}
+				}
+
+				const results = await validateCommonOptions(logger, edge, isPreviewing)
+				validationErrors.push(...results.errors)
+				validationWarnings.push(...results.warnings)
+				}
 		}
 
 		// Check flow
@@ -475,7 +510,7 @@ export class BranchDefs {
 			const flowsTo = new Set()
 			if (def.flowsTo) {
 				if (!Array.isArray(def.flowsTo)) {
-					outErrors.push(`'${def.name}'.flowsTo is not an array`)
+					validationErrors.push(`'${def.name}'.flowsTo is not an array`)
 				}
 				else for (const to of def.flowsTo) {
 					const branchName = names.get(to.toUpperCase())
@@ -483,22 +518,22 @@ export class BranchDefs {
 						flowsTo.add(branchName)
 					}
 					else {
-						outErrors.push(`'${def.name}' flows to unknown branch/alias '${to}'`)
+						validationErrors.push(`'${def.name}' flows to unknown branch/alias '${to}'`)
 					}
 				}
 			}
 
 			if (def.forceFlowTo) {
 				if (!Array.isArray(def.forceFlowTo)) {
-					outErrors.push(`'${def.name}'.forceFlowTo is not an array`)
+					validationErrors.push(`'${def.name}'.forceFlowTo is not an array`)
 				}
 				else for (const to of def.forceFlowTo) {
 					const branchName = names.get(to.toUpperCase())
 					if (!branchName) {
-						outErrors.push(`'${def.name}' force flows to unknown branch/alias '${to}'`)
+						validationErrors.push(`'${def.name}' force flows to unknown branch/alias '${to}'`)
 					}
 					else if (!flowsTo.has(branchName)) {
-						outErrors.push(`'${def.name}' force flows but does not flow to '${to}'`)
+						validationErrors.push(`'${def.name}' force flows but does not flow to '${to}'`)
 					}
 				}
 			}
@@ -508,7 +543,7 @@ export class BranchDefs {
 					await import(`./changelistparsers/${def.changelistParser}`)
 				}
 				catch {
-					outErrors.push(`'${def.name}' specifies custom parser '${def.changelistParser}' but unable to import it`)
+					validationErrors.push(`'${def.name}' specifies custom parser '${def.changelistParser}' but unable to import it`)
 				}
 			}
 		}
@@ -516,33 +551,46 @@ export class BranchDefs {
 		// Check branchspecs for valid branches
 		if (branchGraph.branchspecs) {
 			for (const spec of branchGraph.branchspecs) {
+				let expectedFields = ['from', 'to', 'name']
 				for (const [key, val] of Object.entries(spec)) {
-					if (key !== 'from' && key !== 'to' && key !== 'name') {
-						outErrors.push('Unexpected branchspec property: ' + key)
+					const efLength = expectedFields.length
+					expectedFields = expectedFields.filter(f => f != key)
+					if (efLength == expectedFields.length) {
+						validationErrors.push('Unexpected branchspec property: ' + key)
 					}
 					if (typeof val !== 'string') {
-						outErrors.push(`Branchspec property ${key} is not a string`)
+						validationErrors.push(`Branchspec property ${key} is not a string`)
 					}
 				}
-				if (!spec.from || !spec.to) {
-					outErrors.push(`Invalid branchspec ${spec.name} (requires both to and from fields)`)
+
+				if (expectedFields.length > 0) {
+					validationErrors.push(`Branchspec missing field(s): ${expectedFields.join(', ')}`)
+				}
+
+				if (spec.from && !names.has(spec.from.toUpperCase())) {
+					validationErrors.push(`From-Branch ${spec.from} not found in branchspec ${spec.name}`)
 				}
 
-				if (!names.has(spec.from.toUpperCase())) {
-					outErrors.push(`From-Branch ${spec.from} not found in branchspec ${spec.name}`)
+				if (spec.to && !names.has(spec.to.toUpperCase())) {
+					validationErrors.push(`To-Branch ${spec.to} not found in branchspec ${spec.name}`)
 				}
 
-				if (!names.has(spec.to.toUpperCase())) {
-					outErrors.push(`To-Branch ${spec.to} not found in branchspec ${spec.name}`)
+				if (spec.name && !await PerforceContext.getServerContext(logger).branchExists(spec.name)) {
+					if (isPreviewing) {
+						validationWarnings.push(`Branchspec ${spec.name} does not exist.`)
+					} 
+					else {
+						validationErrors.push(`Branchspec ${spec.name} does not exist.`)
+					}
 				}
 			}
 		}
 
-		if (outErrors.length > 0) {
-			console.log(outErrors)
-			return {branchGraphDef: null, config: defaultConfigForWholeBot}
+		if (validationErrors.length > 0) {
+			console.log(validationErrors)
+			return {branchGraphDef: null, config: defaultConfigForWholeBot, validationErrors, validationWarnings}
 		}
 
-		return {branchGraphDef: branchGraph, config: defaultConfigForWholeBot}
+		return {branchGraphDef: branchGraph, config: defaultConfigForWholeBot, validationErrors, validationWarnings}
 	}
 }
diff --git a/Engine/Extras/RoboMerge/v3/src/robo/graphbot.ts b/Engine/Extras/RoboMerge/v3/src/robo/graphbot.ts
--- a/Engine/Extras/RoboMerge/v3/src/robo/graphbot.ts
+++ b/Engine/Extras/RoboMerge/v3/src/robo/graphbot.ts
@@ -64,10 +64,9 @@ export class GraphBot implements BotEventHandler {
 		graphBot.botLogger.info(`Loading branch map from ${branchSettingsPath}`)
 		const fileText = require('fs').readFileSync(branchSettingsPath, 'utf8')
 
-		const validationErrors: string[] = []
-		const result = await BranchDefs.parseAndValidate(graphBot.botLogger,validationErrors, fileText)
+		const result = await BranchDefs.parseAndValidate(graphBot.botLogger, fileText)
 		if (!result.branchGraphDef) {
-			throw new Error(validationErrors.length === 0 ? 'Failed to parse' : validationErrors.join('\n'))
+			throw new Error(result.validationErrors.length === 0 ? 'Failed to parse' : result.validationErrors.join('\n'))
 		}
 
 		graphBot.branchGraph.config = result.config
@@ -85,7 +84,7 @@ export class GraphBot implements BotEventHandler {
 		}
 
 		if (!result.branchGraphDef) {
-			error = validationErrors.length === 0 ? 'Failed to parse' : validationErrors.join('\n');
+			error = result.validationErrors.length === 0 ? 'Failed to parse' : result.validationErrors.join('\n');
 		}
 
 		// start empty bot on error - can be fixed up by branch definition check-in
diff --git a/Engine/Extras/RoboMerge/v3/src/robo/preview.ts b/Engine/Extras/RoboMerge/v3/src/robo/preview.ts
--- a/Engine/Extras/RoboMerge/v3/src/robo/preview.ts
+++ b/Engine/Extras/RoboMerge/v3/src/robo/preview.ts
@@ -34,17 +34,22 @@ export async function getPreview(cl: number, singleBot?: string) {
 		}
 		const fileText = await p4.print(`${path}@=${cl}`)
 
-		let validationErrors: string[] = []
-		const result = await BranchDefs.parseAndValidate(logger, validationErrors, fileText, true)
+		const result = await BranchDefs.parseAndValidate(logger, fileText, true)
 
 		const errorPrefix = `\n\t${bot} validation failed: `
+		if (result.validationWarnings.length == 1) {
+			status.addWarning(errorPrefix + result.validationWarnings[0])
+		} else if (result.validationWarnings.length > 1) {
+			status.addWarning(errorPrefix + result.validationWarnings.map(warn => `\n\t\t${warn}`).join(''))
+		}
+
 		if (!result.branchGraphDef) {
-			if (validationErrors.length == 0) {
+			if (result.validationErrors.length == 0) {
 				errors.push(errorPrefix + 'unknown error')
-			} else if (validationErrors.length == 1) {
-				errors.push(errorPrefix + validationErrors[0])
+			} else if (result.validationErrors.length == 1) {
+				errors.push(errorPrefix + result.validationErrors[0])
 			} else {
-				errors.push(errorPrefix + validationErrors.map(err => `\n\t\t${err}`).join(''))
+				errors.push(errorPrefix + result.validationErrors.map(err => `\n\t\t${err}`).join(''))
 			}
 			continue
 		}
@@ -65,7 +70,13 @@ export async function getPreview(cl: number, singleBot?: string) {
 	}
 
 	if (errors.length > 0) {
-		throw new Error(errors.join('\n\n'))
+		let error = errors.join('\n\n')
+		const warnings = status.getWarnings()
+		if (warnings.length > 0) {
+			error += '\n\nWarning:'
+			error += status.getWarnings().map(warning => `${warning}`).join('\n\n')
+		}
+		throw new Error(error)
 	}
 	return status
 }
diff --git a/Engine/Extras/RoboMerge/v3/src/robo/status.ts b/Engine/Extras/RoboMerge/v3/src/robo/status.ts
--- a/Engine/Extras/RoboMerge/v3/src/robo/status.ts
+++ b/Engine/Extras/RoboMerge/v3/src/robo/status.ts
@@ -21,6 +21,14 @@ export class Status {
 		}
 	}
 
+	addWarning(warning: string) {
+		this.validationWarnings.push(warning)
+	}
+
+	getWarnings() {
+		return this.validationWarnings
+	}
+
 	reportBotState(botName: string, state: GraphBotState) {
 		this.botStates.push([botName, state])
 	}
@@ -125,6 +133,7 @@ export class Status {
 
 	private allBranches: BranchStatus[] = []
 	private botStates: [string, GraphBotState][] = []
+	private validationWarnings: string[] = []
 }
 
 export function preprocess(data: UserStatusData) {
