Fix syntax errors not reporting correctly in additional channels

diff --git a/Engine/Extras/RoboMerge/v3/src/robo/notifications.ts b/Engine/Extras/RoboMerge/v3/src/robo/notifications.ts
--- a/Engine/Extras/RoboMerge/v3/src/robo/notifications.ts
+++ b/Engine/Extras/RoboMerge/v3/src/robo/notifications.ts
@@ -632,7 +632,15 @@ export class BotNotifications implements BotEventHandler {
 			}
 		}
 		else {
-			channelsToPostTo.push(message.channel)
+			if (sourceBranch.config.additionalSlackChannelForBlockages) {
+				if (!sourceBranch.config.postMessagesToAdditionalChannelOnly) {
+					channelsToPostTo.push(message.channel)
+				}
+				channelsToPostTo.push(sourceBranch.config.additionalSlackChannelForBlockages)
+			}
+			else {
+				channelsToPostTo.push(message.channel)
+			}
 			messagesToPost.push({ message })
 
 			const syntaxErrorMessage: SlackMessage = {
