Allow shelves to be created for transfer errors

diff --git a/Engine/Extras/RoboMerge/v3/public/js/boilerplate.js b/Engine/Extras/RoboMerge/v3/public/js/boilerplate.js
--- a/Engine/Extras/RoboMerge/v3/public/js/boilerplate.js
+++ b/Engine/Extras/RoboMerge/v3/public/js/boilerplate.js
@@ -1461,7 +1461,8 @@ function renderActionsCell_Common(actionCell, data, operationFunction, operation
 			
 			// Can only create shelves for merge conflicts and commit failures
 			if (conflict.kind !== 'Merge conflict' &&
-				conflict.kind !== 'Commit failure') {
+				conflict.kind !== 'Commit failure' &&
+				conflict.kind !== 'Transfer error') {
 				shelfOption.addClass("disabled")
 				shelfOption.off('click')
 				// Bootstrap magic converts 'title' to 'data-original-title'
diff --git a/Engine/Extras/RoboMerge/v3/src/common/perforce.ts b/Engine/Extras/RoboMerge/v3/src/common/perforce.ts
--- a/Engine/Extras/RoboMerge/v3/src/common/perforce.ts
+++ b/Engine/Extras/RoboMerge/v3/src/common/perforce.ts
@@ -1032,7 +1032,7 @@ export class PerforceContext {
 		}
 	}
 
-	async transfer(targetP4: PerforceContext, roboWorkspace: RoboWorkspace, source: IntegrationSource, dest_changelist: number, target: IntegrationTarget)
+	async transfer(targetP4: PerforceContext, roboWorkspace: RoboWorkspace, source: IntegrationSource, dest_changelist: number, target: IntegrationTarget, printIfFailures: boolean)
 	: Promise<string[]> {
 
 		// TODO: support branchspecs
@@ -1160,7 +1160,7 @@ export class PerforceContext {
 
 		const MAX_PRINT_FAILURES = 10
 		let failedPrints = 0
-		if (failures.length == 0) {
+		if (printIfFailures || failures.length == 0) {
 			const PARALLEL_PRINTS = 25
 			const doPrints = async (startIndex: number) => {
 				for (let i=startIndex; i < printCommands.length; i += PARALLEL_PRINTS) {
diff --git a/Engine/Extras/RoboMerge/v3/src/robo/edgebot.ts b/Engine/Extras/RoboMerge/v3/src/robo/edgebot.ts
--- a/Engine/Extras/RoboMerge/v3/src/robo/edgebot.ts
+++ b/Engine/Extras/RoboMerge/v3/src/robo/edgebot.ts
@@ -247,7 +247,7 @@ class EdgeBotImpl extends PerforceStatefulBot {
 				results.push({name: file.fromFile, kind: "delete"})
 			}
 			else if (file.resolveType.toLowerCase() === "content") {
-				results.push({name: file.fromFile, kind: "merge"})
+				results.push({name: file.fromFile, kind: "merge", contentResolveType: file.contentResolveType})
 			} else {
 				// We really shouldn't get this kind, but it's better to display unknown than skip displaying the file
 				results.push({name: file.fromFile, kind: "unknown"})
@@ -286,7 +286,7 @@ class EdgeBotImpl extends PerforceStatefulBot {
 			const owner = getIntegrationOwner(pending) || pending.change.author
 			const shelfMsg = `${owner}, please merge this change by hand.\nMore info at ${this.sourceNode.getBotUrl()}\n\n` + failure.description
 			pending.change.additionalDescriptionText = `#ROBOMERGE-CONFLICT from-shelf\n`
-			await this.shelveChangelist(pending, false, shelfMsg)
+			await this.shelveChangelist(pending, {reason: shelfMsg})
 			this.edgeBotLogger.info(`${logMessage}. Shelved CL ${pending.newCl} for ${owner} to resolve manually (from reconsider).`)
 			return
 		}
@@ -658,14 +658,24 @@ class EdgeBotImpl extends PerforceStatefulBot {
 
 		this.currentIntegrationStartTimestamp = Date.now()
 		await targetServerP4.sync(info.targetWorkspace, this.targetBranch.rootPath + '#0')
-		const failures = await this.p4.transfer(targetServerP4, info.targetWorkspace, source, changenum, integTarget)
+		const failures = await this.p4.transfer(targetServerP4, info.targetWorkspace, source, changenum, integTarget, info.forceCreateAShelf)
 
 		const pending: PendingChange = {change: info, action: target, newCl: changenum}
 
 		let description = failures.join('\n')
 		let failure: Failure | null = null
 
-		if (failures.length == 0) {
+		if (info.forceCreateAShelf) {
+			// the user requested manual merge
+			await this.shelveChangelist(pending, {targetServerP4})
+
+			if (!pending.change.sendNoShelfNotification) {
+				await this.sourceNode.notifyShelfRequester(pending)
+			}
+			
+			return new EdgeIntegrationDetails('shelved', pending.newCl)			
+		}
+		else if (failures.length == 0) {
 			// no failures, try to submit
 			const result = await this.submitChangelist(pending, 0, targetServerP4)
 			if (result.result !== 'error') {
@@ -757,6 +767,7 @@ class EdgeBotImpl extends PerforceStatefulBot {
 				failure.details = conflicts
 					.map(({name, kind}) => `${name} (${kind} conflict)`)
 					.join('\n')
+				failure.additionalInfo = conflicts
 			}
 		}
 		else
@@ -796,7 +807,7 @@ class EdgeBotImpl extends PerforceStatefulBot {
 				const approval = this.options.approval
 
 				// shelve
-				await this.shelveChangelist(pending, true)
+				await this.shelveChangelist(pending, {forApproval: true})
 	
 				// still notify as a 'failure' to trigger normal blockage mechanism
 				failure = { kind: 'Approval required', description: approval.description }
@@ -928,7 +939,12 @@ class EdgeBotImpl extends PerforceStatefulBot {
 		return new EdgeIntegrationDetails('ok', finalCl)
 	}
 
-	private async shelveChangelist(pending: PendingChange, forApproval?: boolean, reason?: string) {
+	private async shelveChangelist(pending: PendingChange, shelveOpts?: {forApproval?: boolean, reason?: string, targetServerP4?: PerforceContext}) {
+
+		const forApproval: boolean = shelveOpts?.forApproval || false
+		const reason: string | undefined = shelveOpts?.reason
+		const p4 = shelveOpts?.targetServerP4 || this.p4
+
 		const changenum = pending.newCl
 		const owner = getIntegrationOwner(pending) || pending.change.author
 		this._log_action(`Shelving CL ${changenum} (change owned by ${owner})`)
@@ -956,15 +972,15 @@ class EdgeBotImpl extends PerforceStatefulBot {
 		if (isUserAKnownBot(owner) && !pending.change.forceCreateAShelf && !forApproval) {
 			// revert the changes locally
 			this._log_action(`Reverting shelf due to '${owner}' being a known bot`)
-			await this.p4.revert(destRoboWorkspace, changenum, undefined, edgeServerAddress)
+			await p4.revert(destRoboWorkspace, changenum, undefined, edgeServerAddress)
 			failed = true
 		}
 		else {
 			// edit the CL description
-			await this.p4.editDescription(destRoboWorkspace, changenum, final_desc, edgeServerAddress)
+			await p4.editDescription(destRoboWorkspace, changenum, final_desc, edgeServerAddress)
 
 			// shelve the files as we see them (should trigger a codereview email)
-			if (!await this.p4.shelve(destRoboWorkspace, changenum, edgeServerAddress)) {
+			if (!await p4.shelve(destRoboWorkspace, changenum, edgeServerAddress)) {
 				failed = true
 			}
 		}
@@ -973,7 +989,7 @@ class EdgeBotImpl extends PerforceStatefulBot {
 			// abort abort!
 			// delete the cl
 			this._log_action(`Deleting CL ${changenum}`)
-			await this.p4.deleteCl(destRoboWorkspace, changenum, edgeServerAddress)
+			await p4.deleteCl(destRoboWorkspace, changenum, edgeServerAddress)
 
 			pending.newCl = -1
 
@@ -983,7 +999,7 @@ class EdgeBotImpl extends PerforceStatefulBot {
 
 		// revert the changes locally
 		this._log_action(`Reverting CL ${changenum} locally. (conflict owner: ${owner})`)
-		await this.p4.revert(destRoboWorkspace, changenum, [], edgeServerAddress)
+		await p4.revert(destRoboWorkspace, changenum, [], edgeServerAddress)
 
 		// figure out what workspace to put it in
 		const branch_stream = this.targetBranch.stream ? this.targetBranch.stream.toLowerCase() : undefined
@@ -995,7 +1011,7 @@ class EdgeBotImpl extends PerforceStatefulBot {
 		// Find a suitable workspace from one of the owner's workspaces
 		else if (!forApproval) {
 			// use p4.find_workspaces to find a workspace (owned by the user) for this change if this is a stream branch
-			const targetWorkspaceDef = await p4util.chooseBestWorkspaceForUser(this.p4, owner, branch_stream)
+			const targetWorkspaceDef = await p4util.chooseBestWorkspaceForUser(p4, owner, branch_stream)
 			if (targetWorkspaceDef) {
 				targetWorkspace = targetWorkspaceDef.client
 				if (targetWorkspaceDef.Stream) {
@@ -1028,7 +1044,7 @@ class EdgeBotImpl extends PerforceStatefulBot {
 
 		// edit the owner to the author so they can resolve and submit themselves
 		let shelfOwner = forApproval ? pending.change.author : owner
-		await this.p4.editOwner(destRoboWorkspace, changenum, shelfOwner, opts)
+		await p4.editOwner(destRoboWorkspace, changenum, shelfOwner, opts)
 	}
 
 	onGlobalChange(info: ChangeInfo) {
