Refactor track changes in to its own file

diff --git a/Engine/Extras/RoboMerge/v3/src/robo/ipc.ts b/Engine/Extras/RoboMerge/v3/src/robo/ipc.ts
--- a/Engine/Extras/RoboMerge/v3/src/robo/ipc.ts
+++ b/Engine/Extras/RoboMerge/v3/src/robo/ipc.ts
@@ -1,7 +1,7 @@
 // Copyright Epic Games, Inc. All Rights Reserved.
 
 import { ContextualLogger, isNpmLogLevel, NpmLogLevelCompare, NpmLogLevelValues } from '../common/logger';
-import { DescribeResult, getRunningPerforceCommands, PerforceContext } from '../common/perforce';
+import { getRunningPerforceCommands, PerforceContext } from '../common/perforce';
 import { Trace } from '../new/graph';
 import { IPCControls, EdgeBotInterface, NodeBotInterface } from './bot-interfaces';
 import { OperationResult } from './branch-interfaces';
@@ -9,6 +9,7 @@ import { GraphBot } from './graphbot';
 import { RoboMerge } from './robo';
 import { roboAnalytics } from './roboanalytics';
 import { Status } from './status';
+import { trackChange } from './trackchange';
 import { getPreview } from './preview'
 import * as p4util from '../common/p4util';
 
@@ -27,10 +28,6 @@ export interface Message {
 // roboserver.ts -- getQueryFromSecure()
 type Query = {[key: string]: string};
 
-const RobomergeMethodStrings = ['initialSubmit','merge_with_conflict','automerge'] as const;
-type RobomergeMethods = typeof RobomergeMethodStrings[number];
-type MergeMethod = RobomergeMethods|'populate'|'manual_merge'
-
 export type OperationReturnType = {
 	statusCode: number
 	message: string // Goal: to provide meaningful error messaging to the end user
@@ -186,200 +183,12 @@ export class IPC {
 			return {statusCode: 503, message: 'Service is not ready to process request'}
 		}
 
-		let userTags = new Set<string>()
+		const userTags = new Set<string>()
 		for (const tag in tagsObj) {
 			userTags.add(tag)
 		}
-
-		const streamFilter = (() => {
-			const streamsParam: string|undefined = queryObj.streams;
-			return (streamsParam ? streamsParam.toUpperCase().replace('*','.*').split(',').map(s => new RegExp(s)) : [])
-		})() 
-		const botFilter = (() => { 
-			const botsParam = queryObj.bots
-			return (botsParam ? botsParam.toUpperCase().split(',') : [])
-		})()
-		const depotFilter = (() => { 
-			const depotsParam: string = queryObj.depots
-			return (depotsParam ? depotsParam.toUpperCase().split(',').map(depot => `//${depot}`) : [])
-		})()
-
-		const graph = this.robo.graph.graph
-
-		const getStreamFromPath = (path: string) => {
-			const match = path.match(/\/\/[^\/]*\/[^\/]*/)
-			return (match ? match[0] : null)
-		}
-
-		const getNode = (desc: DescribeResult) => {
-			if (desc.path) {
-				const stream = getStreamFromPath(desc.path)
-				if (stream) {
-					try {
-						return graph.findNodeForStream(stream)
-					}
-					catch {						
-					}
-				}
-			}
-			return null
-		}
-
-		const getMergeMethod = (desc: string): MergeMethod => {
-			if (desc.includes("#ROBOMERGE-CONFLICT")) {
-				return 'merge_with_conflict'
-			} else if (desc.includes("#ROBOMERGE-SOURCE")) {
-				return 'automerge'
-			} else if (desc.includes("Populate")) {
-				return 'populate'
-			} else {
-				return 'manual_merge'
-			}
-		}
-
-		let p4 = PerforceContext.getServerContext(this.robo.roboMergeLogger)
-		let changes = new Map<number, any>()
-
-		let data: any = {originalCL: cl, changes: {}}
-
-		let gatherCLInfo = async (cl: number, opts?: any) => {
-			let desc = await p4.describe(cl, 1)
-			const mergeMethod = cl != data.originalCL ? getMergeMethod(desc.description) : 'initialSubmit'
-			const clNode = getNode(desc)
-			if (clNode || (RobomergeMethodStrings as readonly string[]).includes(mergeMethod)) {
-				changes.set(cl, {desc, node: clNode, sourceCL: opts ? opts.sourceCL : null, destCLs: (opts && opts.lastCL ? [opts.lastCL] : []) })
-				return true
-			}
-			return false
-		}
-
-		await gatherCLInfo(cl)
-
-		const sourceMatch = (changes.get(data.originalCL).desc as DescribeResult).description.match(/#ROBOMERGE-SOURCE: (.*)\n/g)
-		if (sourceMatch) {
-			let lastCL = cl
-			const matches = Array.from(sourceMatch[0].matchAll(/CL (\d+)/g))
-			for (let i=matches.length-1; i>=0;i--) {
-				let sourceCL = parseInt(matches[i][1])
-				if (i == 0) {
-					data.originalCL = sourceCL
-				}
-				await gatherCLInfo(sourceCL, {lastCL})
-				lastCL = sourceCL
-			}
-		}
-
-		const isFTE = userTags.has("fte")
-
-		let clsToConsider: number[] = [data.originalCL]
-		while (clsToConsider.length > 0) {
-			const clToConsider = clsToConsider[0]
-			clsToConsider = clsToConsider.slice(1)
-
-			let changeToConsider = changes.get(clToConsider)
-			if (!changeToConsider) {
-				continue
-			}
-
-			let includeInResults = false
-			let hasAutomergeTarget = false
-			let streamDisplayName = ""
-			if (changeToConsider.node) {
-				let edges = graph.getEdgesForNode(changeToConsider.node)
-				for (let edge of edges) {
-					if (!includeInResults) {
-						const bot = edge.sourceAnnotation as NodeBotInterface
-						if (botFilter.length == 0 || botFilter.includes(bot.branchGraph.botname)) {
-							if (Status.includeBranch(bot.branchGraph.config.visibility, userTags, this.ipcLogger)) {
-								streamDisplayName = changeToConsider.node.stream
-								includeInResults = true
-							}
-						}
-					}
-					if (!hasAutomergeTarget && edge.flags.has('automatic')) {
-						if (typeof changeToConsider.node === 'string') {
-							hasAutomergeTarget = edge.source.stream == changeToConsider.node
-						}
-						else {
-							hasAutomergeTarget = edge.source.stream == changeToConsider.node.stream
-						}
-					}
-					if (includeInResults && hasAutomergeTarget) {
-						break
-					}
-				}
-			} else if (isFTE && botFilter.length == 0) {
-				includeInResults = true
-				if (changeToConsider.desc.path) {
-					streamDisplayName = getStreamFromPath(changeToConsider.desc.path) || changeToConsider.desc.path
-				} else if (changeToConsider.desc.entries.length > 0) {
-					streamDisplayName = getStreamFromPath(changeToConsider.desc.entries[0].depotFile) || streamDisplayName
-				}
-			}
-
-			if (streamDisplayName.length == 0) {
-				streamDisplayName = "//****/****"
-			}
-
-			if (includeInResults)
-			{
-				const upperSteamDisplayName = streamDisplayName.toUpperCase()
-				includeInResults = depotFilter.length == 0 || depotFilter.some(depot => upperSteamDisplayName.startsWith(depot))
-				includeInResults = includeInResults && (streamFilter.length == 0 || streamFilter.some(re => upperSteamDisplayName.match(re)))
-			}
-
-			// Can't cache this because the passed in changelist is incorrectly labelled initialSubmit the first time
-			// it is considered
-			const mergeMethod = clToConsider != data.originalCL ? getMergeMethod(changeToConsider.desc.description) : 'initialSubmit'
-
-			for (let i=0; i < changeToConsider.desc.entries.length; i++) {
-				const entry = changeToConsider.desc.entries[i]
-				// integrated for move/delete will point at the paired move/add not where it was merged to in another stream, so not useful to evaluate it
-				if (entry.action != 'move/delete') {
-					const integrated = await p4.integrated(null, entry.depotFile, {intoOnly: true, startCL: clToConsider})
-					if (integrated.length > 0) {
-						for (let integ of integrated) {
-							const startToRev = integ.startToRev == "#none" ? 0 : parseInt(integ.startToRev.slice(1))
-							const endToRev = parseInt(integ.endToRev.slice(1))
-							if (entry.rev <= endToRev && entry.rev > startToRev)
-							{
-								const destChange = changes.get(integ.change)
-								if (destChange) {
-									destChange.sourceCL = clToConsider
-								} else {
-									changeToConsider.destCLs.push(integ.change)
-									await gatherCLInfo(integ.change, {sourceCL: clToConsider})
-								}
-							}
-						}
-						break
-					}
-				}
-				if (hasAutomergeTarget && changeToConsider.desc.entries.length == 1 && 
-							(RobomergeMethodStrings as readonly string[]).includes(mergeMethod)) {
-					// If we only have 1 entry and we didn't get integration info off of it
-					// and the graph suggests we are expecting there could be other changes
-					// get more of the files. If the path is a root directory try and get a sampling 
-					// of files, otherwise just get an additional block to evaluate
-					if (changeToConsider.desc.path.endsWith("/...")) {
-						const dirs = await p4.dirs(`${changeToConsider.desc.path.slice(0,-3)}*@=${clToConsider}`)
-						for (const dir of dirs) {
-							changeToConsider.desc.entries.push(...(await p4.files(`${dir}/...@=${clToConsider}`, 1)))
-						}
-					}
-					else {
-						changeToConsider.desc = await p4.describe(clToConsider, 100)
-					}
-				}
-			}
-			clsToConsider = clsToConsider.concat(changeToConsider.destCLs)
-
-			if (includeInResults) {
-				data.changes[`${clToConsider}`] = {streamDisplayName, mergeMethod, sourceCL: changeToConsider.sourceCL}
-			}
-		}
-
-		return {...OPERATION_SUCCESS, data}
+	
+		return {...OPERATION_SUCCESS, data: await trackChange(this.robo, this.ipcLogger, cl, userTags, queryObj)}
 	}
 
 	private async traceRoute(query: Query, tagsObj?: any): Promise<OperationReturnType> {
diff --git a/Engine/Extras/RoboMerge/v3/src/robo/trackchange.ts b/Engine/Extras/RoboMerge/v3/src/robo/trackchange.ts
new file mode 100644
--- /dev/null
+++ b/Engine/Extras/RoboMerge/v3/src/robo/trackchange.ts
@@ -0,0 +1,207 @@
+// Copyright Epic Games, Inc. All Rights Reserved.
+
+import { ContextualLogger } from '../common/logger';
+import { DescribeResult, PerforceContext } from '../common/perforce';
+import { NodeBotInterface } from './bot-interfaces';
+import { RoboMerge } from './robo';
+import { Status } from './status';
+
+const RobomergeMethodStrings = ['initialSubmit','merge_with_conflict','automerge'] as const;
+type RobomergeMethods = typeof RobomergeMethodStrings[number];
+type MergeMethod = RobomergeMethods|'populate'|'manual_merge'
+
+interface TrackFilters {
+    streams?: string
+    bots?: string
+    depots?: string
+}
+
+export async function trackChange(robo: RoboMerge, logger: ContextualLogger, cl: number, userTags: Set<string>, filters: TrackFilters) {
+
+    const streamFilter = (() => {
+        return (filters.streams ? filters.streams.toUpperCase().replace('*','.*').split(',').map(s => new RegExp(s)) : [])
+    })() 
+    const botFilter = (() => { 
+        return (filters.bots ? filters.bots.toUpperCase().split(',') : [])
+    })()
+    const depotFilter = (() => { 
+        return (filters.depots ? filters.depots.toUpperCase().split(',').map(depot => `//${depot}`) : [])
+    })()
+
+    const graph = robo.graph.graph
+
+    const getStreamFromPath = (path: string) => {
+        const match = path.match(/\/\/[^\/]*\/[^\/]*/)
+        return (match ? match[0] : null)
+    }
+
+    const getNode = (desc: DescribeResult) => {
+        if (desc.path) {
+            const stream = getStreamFromPath(desc.path)
+            if (stream) {
+                try {
+                    return graph.findNodeForStream(stream)
+                }
+                catch {						
+                }
+            }
+        }
+        return null
+    }
+
+    const getMergeMethod = (desc: string): MergeMethod => {
+        if (desc.includes("#ROBOMERGE-CONFLICT")) {
+            return 'merge_with_conflict'
+        } else if (desc.includes("#ROBOMERGE-SOURCE")) {
+            return 'automerge'
+        } else if (desc.includes("Populate")) {
+            return 'populate'
+        } else {
+            return 'manual_merge'
+        }
+    }
+
+    let p4 = PerforceContext.getServerContext(robo.roboMergeLogger)
+    let changes = new Map<number, any>()
+
+    let data: any = {originalCL: cl, changes: {}}
+
+    let gatherCLInfo = async (cl: number, opts?: any) => {
+        let desc = await p4.describe(cl, 1)
+        const mergeMethod = cl != data.originalCL ? getMergeMethod(desc.description) : 'initialSubmit'
+        const clNode = getNode(desc)
+        if (clNode || (RobomergeMethodStrings as readonly string[]).includes(mergeMethod)) {
+            changes.set(cl, {desc, node: clNode, sourceCL: opts ? opts.sourceCL : null, destCLs: (opts && opts.lastCL ? [opts.lastCL] : []) })
+            return true
+        }
+        return false
+    }
+
+    await gatherCLInfo(cl)
+
+    const sourceMatch = (changes.get(data.originalCL).desc as DescribeResult).description.match(/#ROBOMERGE-SOURCE: (.*)\n/g)
+    if (sourceMatch) {
+        let lastCL = cl
+        const matches = Array.from(sourceMatch[0].matchAll(/CL (\d+)/g))
+        for (let i=matches.length-1; i>=0;i--) {
+            let sourceCL = parseInt(matches[i][1])
+            if (i == 0) {
+                data.originalCL = sourceCL
+            }
+            await gatherCLInfo(sourceCL, {lastCL})
+            lastCL = sourceCL
+        }
+    }
+
+    const isFTE = userTags.has("fte")
+
+    let clsToConsider: number[] = [data.originalCL]
+    while (clsToConsider.length > 0) {
+        const clToConsider = clsToConsider[0]
+        clsToConsider = clsToConsider.slice(1)
+
+        let changeToConsider = changes.get(clToConsider)
+        if (!changeToConsider) {
+            continue
+        }
+
+        let includeInResults = false
+        let hasAutomergeTarget = false
+        let streamDisplayName = ""
+        if (changeToConsider.node) {
+            let edges = graph.getEdgesForNode(changeToConsider.node)
+            for (let edge of edges) {
+                if (!includeInResults) {
+                    const bot = edge.sourceAnnotation as NodeBotInterface
+                    if (botFilter.length == 0 || botFilter.includes(bot.branchGraph.botname)) {
+                        if (Status.includeBranch(bot.branchGraph.config.visibility, userTags, logger)) {
+                            streamDisplayName = changeToConsider.node.stream
+                            includeInResults = true
+                        }
+                    }
+                }
+                if (!hasAutomergeTarget && edge.flags.has('automatic')) {
+                    if (typeof changeToConsider.node === 'string') {
+                        hasAutomergeTarget = edge.source.stream == changeToConsider.node
+                    }
+                    else {
+                        hasAutomergeTarget = edge.source.stream == changeToConsider.node.stream
+                    }
+                }
+                if (includeInResults && hasAutomergeTarget) {
+                    break
+                }
+            }
+        } else if (isFTE && botFilter.length == 0) {
+            includeInResults = true
+            if (changeToConsider.desc.path) {
+                streamDisplayName = getStreamFromPath(changeToConsider.desc.path) || changeToConsider.desc.path
+            } else if (changeToConsider.desc.entries.length > 0) {
+                streamDisplayName = getStreamFromPath(changeToConsider.desc.entries[0].depotFile) || streamDisplayName
+            }
+        }
+
+        if (streamDisplayName.length == 0) {
+            streamDisplayName = "//****/****"
+        }
+
+        if (includeInResults)
+        {
+            const upperSteamDisplayName = streamDisplayName.toUpperCase()
+            includeInResults = depotFilter.length == 0 || depotFilter.some(depot => upperSteamDisplayName.startsWith(depot))
+            includeInResults = includeInResults && (streamFilter.length == 0 || streamFilter.some(re => upperSteamDisplayName.match(re)))
+        }
+
+        // Can't cache this because the passed in changelist is incorrectly labelled initialSubmit the first time
+        // it is considered
+        const mergeMethod = clToConsider != data.originalCL ? getMergeMethod(changeToConsider.desc.description) : 'initialSubmit'
+
+        for (let i=0; i < changeToConsider.desc.entries.length; i++) {
+            const entry = changeToConsider.desc.entries[i]
+            // integrated for move/delete will point at the paired move/add not where it was merged to in another stream, so not useful to evaluate it
+            if (entry.action != 'move/delete') {
+                const integrated = await p4.integrated(null, entry.depotFile, {intoOnly: true, startCL: clToConsider})
+                if (integrated.length > 0) {
+                    for (let integ of integrated) {
+                        const startToRev = integ.startToRev == "#none" ? 0 : parseInt(integ.startToRev.slice(1))
+                        const endToRev = parseInt(integ.endToRev.slice(1))
+                        if (entry.rev <= endToRev && entry.rev > startToRev)
+                        {
+                            const destChange = changes.get(integ.change)
+                            if (destChange) {
+                                destChange.sourceCL = clToConsider
+                            } else {
+                                changeToConsider.destCLs.push(integ.change)
+                                await gatherCLInfo(integ.change, {sourceCL: clToConsider})
+                            }
+                        }
+                    }
+                    break
+                }
+            }
+            if (hasAutomergeTarget && changeToConsider.desc.entries.length == 1 && 
+                        (RobomergeMethodStrings as readonly string[]).includes(mergeMethod)) {
+                // If we only have 1 entry and we didn't get integration info off of it
+                // and the graph suggests we are expecting there could be other changes
+                // get more of the files. If the path is a root directory try and get a sampling 
+                // of files, otherwise just get an additional block to evaluate
+                if (changeToConsider.desc.path.endsWith("/...")) {
+                    const dirs = await p4.dirs(`${changeToConsider.desc.path.slice(0,-3)}*@=${clToConsider}`)
+                    for (const dir of dirs) {
+                        changeToConsider.desc.entries.push(...(await p4.files(`${dir}/...@=${clToConsider}`, 1)))
+                    }
+                }
+                else {
+                    changeToConsider.desc = await p4.describe(clToConsider, 100)
+                }
+            }
+        }
+        clsToConsider = clsToConsider.concat(changeToConsider.destCLs)
+
+        if (includeInResults) {
+            data.changes[`${clToConsider}`] = {streamDisplayName, mergeMethod, sourceCL: changeToConsider.sourceCL}
+        }
+    }
+
+    return data
+}
\ No newline at end of file
