Remove in progress changes that got included in previous change

diff --git a/Engine/Extras/RoboMerge/v3/src/robo/edgebot.ts b/Engine/Extras/RoboMerge/v3/src/robo/edgebot.ts
--- a/Engine/Extras/RoboMerge/v3/src/robo/edgebot.ts
+++ b/Engine/Extras/RoboMerge/v3/src/robo/edgebot.ts
@@ -247,7 +247,7 @@ class EdgeBotImpl extends PerforceStatefulBot {
 				results.push({name: file.fromFile, kind: "delete"})
 			}
 			else if (file.resolveType.toLowerCase() === "content") {
-				results.push({name: file.fromFile, kind: "merge", contentResolveType: file.contentResolveType})
+				results.push({name: file.fromFile, kind: "merge"})
 			} else {
 				// We really shouldn't get this kind, but it's better to display unknown than skip displaying the file
 				results.push({name: file.fromFile, kind: "unknown"})
@@ -767,7 +767,6 @@ class EdgeBotImpl extends PerforceStatefulBot {
 				failure.details = conflicts
 					.map(({name, kind}) => `${name} (${kind} conflict)`)
 					.join('\n')
-				failure.additionalInfo = conflicts
 			}
 		}
 		else
