Allow changes to gate files to update visually in preview mode

diff --git a/Engine/Extras/RoboMerge/v3/src/robo/edgebot.ts b/Engine/Extras/RoboMerge/v3/src/robo/edgebot.ts
--- a/Engine/Extras/RoboMerge/v3/src/robo/edgebot.ts
+++ b/Engine/Extras/RoboMerge/v3/src/robo/edgebot.ts
@@ -195,6 +195,10 @@ class EdgeBotImpl extends PerforceStatefulBot {
 	async tick() {
 		// Update the Last Good Change from the gate file in Perforce
 		this.gate.tick()
+		if (this.sourceNode.previewMode) {
+		    return true
+		}
+		// Add any actions that should occur only if not in preview mode here
 		return true
 	}
 
diff --git a/Engine/Extras/RoboMerge/v3/src/robo/nodebot.ts b/Engine/Extras/RoboMerge/v3/src/robo/nodebot.ts
--- a/Engine/Extras/RoboMerge/v3/src/robo/nodebot.ts
+++ b/Engine/Extras/RoboMerge/v3/src/robo/nodebot.ts
@@ -88,7 +88,7 @@ export class NodeBot extends PerforceStatefulBot implements NodeBotInterface {
 
 	constructor(
 		branchDef: Branch,
-		private readonly previewMode: boolean,
+		readonly previewMode: boolean,
 		mailer: Mailer,
 		slackMessages: SlackMessages | undefined,
 		externalUrl: string, 
@@ -242,10 +242,10 @@ export class NodeBot extends PerforceStatefulBot implements NodeBotInterface {
 
 	/** @return true if full tick completed, only used for analytics */
 	async tick() {
+		await Promise.all(Array.from(this.edges.values(), async edgeBot => edgeBot.tick()))
 		if (this.previewMode) {
 			return false
 		}
-		await Promise.all(Array.from(this.edges.values(), async edgeBot => edgeBot.tick()))
 
 		// Pre-tick
 
