Partial cross server track changes
Can only initiate based on CL# on the default server
Ability to track backwards if the original submit is on an external server is limited

diff --git a/Engine/Extras/RoboMerge/v3/public/js/trackchange.js b/Engine/Extras/RoboMerge/v3/public/js/trackchange.js
--- a/Engine/Extras/RoboMerge/v3/public/js/trackchange.js
+++ b/Engine/Extras/RoboMerge/v3/public/js/trackchange.js
@@ -1,34 +1,20 @@
 // Copyright Epic Games, Inc. All Rights Reserved.
 "use strict"
 
-function getMergeMethodString(mergeMethod) {
+function getMergeMethodStyle(mergeMethod) {
 	switch(mergeMethod) {
 		case "automerge":
-			return "Automerge"
+			return {color: "black", style: "solid"}
 		case "initialSubmit":
-			return ""
+			return {color: "", style: ""}
 		case "merge_with_conflict":
-			return "Merge w/ conflict"
+			return {color: "red", style: "dashed"}
 		case "manual_merge":
-			return "Manual merge"
+			return {color: "darkgray", style: "dashed"}
 		case "populate":
-			return "Populate"
-	}
-	return "UNKNOWN CASE: " + mergeMethod
-}
-
-function getMergeMethodColor(mergeMethod) {
-	switch(mergeMethod) {
-		case "automerge":
-			return "black"
-		case "initialSubmit":
-			return ""
-		case "merge_with_conflict":
-			return "red"
-		case "manual_merge":
-			return "darkgray"
-		case "populate":
-			return "blue"
+			return {color: "blue", style: "dashed"}
+		case "transfer":
+			return {color: "orange", style: "solid"}
 	}
 	return "UNKNOWN CASE: " + mergeMethod	
 }
@@ -37,7 +23,8 @@ function getStreamGraphName(streamDisplayName) {
 	if (streamDisplayName.endsWith("/Main")) {
 		return streamDisplayName
 	}
-	return streamDisplayName.split('/').reverse()[0]
+	const match = streamDisplayName.match(/\/\/[^\/]+\/([^\/]+).*/)
+	return match ? match[1] : streamDisplayName
 }
 
 async function getGitHubCommit(cl) {
@@ -74,13 +61,10 @@ async function generateChangeList(dataObj) {
 	const gitHubCommits = await getGitHubCommits(dataObj.changes)
 	let html = '<div style="margin: auto; width: 80%;"><table class="table"><tbody>'
 	for (const cl in dataObj.changes) {
-		console.log()
 		html += '<tr valign="middle">'
 		html += `<td><b>${dataObj.changes[cl].streamDisplayName}</b></td>`
 		html += `<td><a href="https://p4-swarm.epicgames.net/changes/${cl}" target="_blank">CL#${cl}</a></td>`
 		html += `${formatGitHubCommit(gitHubCommits.get(cl))}`
-		//html += `<td style="text-align:center;">${getMergeMethodString(dataObj.changes[cl].mergeMethod)}</td>`
-		//html += `<td style="text-align:center;">${dataObj.changes[cl].sourceCL}</td>`
 		html += '</tr>'
 	}
 	html += "</tbody></table></div>"
@@ -112,7 +96,8 @@ function createGraph(changes) {
 			if (changes[change].mergeMethod == "automerge") {
 				lines.push(`_${changes[change].sourceCL} -> _${change};`)
 			} else {
-				lines.push(`_${changes[change].sourceCL} -> _${change} [color=${getMergeMethodColor(changes[change].mergeMethod)}, style=dashed];`)
+				const mergeMethodStyle = getMergeMethodStyle(changes[change].mergeMethod)
+				lines.push(`_${changes[change].sourceCL} -> _${change} [color=${mergeMethodStyle.color}, style=${mergeMethodStyle.style}];`)
 			}
 		}
 	}
diff --git a/Engine/Extras/RoboMerge/v3/src/new/graph.ts b/Engine/Extras/RoboMerge/v3/src/new/graph.ts
--- a/Engine/Extras/RoboMerge/v3/src/new/graph.ts
+++ b/Engine/Extras/RoboMerge/v3/src/new/graph.ts
@@ -16,7 +16,7 @@ type Success = 'failed' | 'succeeded'
 
 // deliberately not much in here - assuming stream for the moment
 export class Node {
-	constructor(public stream: Stream) // of form //x[/x]*
+	constructor(public stream: Stream, public serverID?: string) // of form //x[/x]*
 	{}
 
 	offerDebugName(name: string) {
@@ -126,12 +126,23 @@ export class Graph {
 
 	findOrCreateStreamNode(depotPath: string | Branch) {
 		const branchId = makeBranchId(depotPath)
-		return setDefault(this.streamNodes, branchId, new Node(branchId))
+		const server = (depotPath as Branch).config.streamServer
+		return setDefault(this.streamNodes, branchId, new Node(branchId, server))
 	}
 
-	findNodeForStream(depotPath: string | Branch) {
+	findNodeForStream(depotPath: string | Branch, allowPartialMatch?: boolean) {
 		const branchId = makeBranchId(depotPath)
-		return this.streamNodes.get(branchId)
+		let nodeForStream = this.streamNodes.get(branchId)
+		if (!nodeForStream && allowPartialMatch) {
+			const partialMatchbranchId = (branchId as string) + '/'
+			for (const [stream, node] of this.streamNodes) {
+				if (stream.startsWith(partialMatchbranchId)) {
+					nodeForStream = node
+					break
+				}
+			}
+		}
+		return nodeForStream
 	}
 
 	findNodeForBranch(branch: Branch) {
diff --git a/Engine/Extras/RoboMerge/v3/src/robo/ipc.ts b/Engine/Extras/RoboMerge/v3/src/robo/ipc.ts
--- a/Engine/Extras/RoboMerge/v3/src/robo/ipc.ts
+++ b/Engine/Extras/RoboMerge/v3/src/robo/ipc.ts
@@ -188,7 +188,7 @@ export class IPC {
 			userTags.add(tag)
 		}
 	
-		return {...OPERATION_SUCCESS, data: await trackChange(this.robo, this.ipcLogger, cl, userTags, queryObj)}
+		return {...OPERATION_SUCCESS, data: await trackChange(this.robo, this.ipcLogger, {cl}, userTags, queryObj)}
 	}
 
 	private async traceRoute(query: Query, tagsObj?: any): Promise<OperationReturnType> {
diff --git a/Engine/Extras/RoboMerge/v3/src/robo/trackchange.ts b/Engine/Extras/RoboMerge/v3/src/robo/trackchange.ts
--- a/Engine/Extras/RoboMerge/v3/src/robo/trackchange.ts
+++ b/Engine/Extras/RoboMerge/v3/src/robo/trackchange.ts
@@ -2,11 +2,12 @@
 
 import { ContextualLogger } from '../common/logger';
 import { DescribeResult, PerforceContext } from '../common/perforce';
+import { Node } from '../new/graph';
 import { NodeBotInterface } from './bot-interfaces';
 import { RoboMerge } from './robo';
 import { Status } from './status';
 
-const RobomergeMethodStrings = ['initialSubmit','merge_with_conflict','automerge'] as const;
+const RobomergeMethodStrings = ['initialSubmit','merge_with_conflict','automerge','transfer'] as const;
 type RobomergeMethods = typeof RobomergeMethodStrings[number];
 type MergeMethod = RobomergeMethods|'populate'|'manual_merge'
 
@@ -16,7 +17,29 @@ interface TrackFilters {
     depots?: string
 }
 
-export async function trackChange(robo: RoboMerge, logger: ContextualLogger, cl: number, userTags: Set<string>, filters: TrackFilters) {
+interface TrackedCL {
+    cl: number
+    serverId?: string
+}
+
+function compareTrackedCL(a: TrackedCL, b: TrackedCL) {
+    return a.cl == b.cl && a.serverId == b.serverId
+}
+
+interface TrackedChange {
+    desc: DescribeResult
+    mergeMethod: MergeMethod
+    node?: Node
+    sourceCL?: TrackedCL
+    destCLs: TrackedCL[]
+}
+
+interface TrackChangeResults {
+    originalCL: TrackedCL,
+    changes: any
+}
+
+export async function trackChange(robo: RoboMerge, logger: ContextualLogger, initialCL: TrackedCL, userTags: Set<string>, filters: TrackFilters) {
 
     const streamFilter = (() => {
         return (filters.streams ? filters.streams.toUpperCase().replace('*','.*').split(',').map(s => new RegExp(s)) : [])
@@ -29,6 +52,7 @@ export async function trackChange(robo: RoboMerge, logger: ContextualLogger, cl:
     })()
 
     const graph = robo.graph.graph
+    let data: TrackChangeResults = {originalCL: initialCL, changes: {}}
 
     const getStreamFromPath = (path: string) => {
         const match = path.match(/\/\/[^\/]*\/[^\/]*/)
@@ -40,17 +64,21 @@ export async function trackChange(robo: RoboMerge, logger: ContextualLogger, cl:
             const stream = getStreamFromPath(desc.path)
             if (stream) {
                 try {
-                    return graph.findNodeForStream(stream)
+                    return graph.findNodeForStream(stream, true)
                 }
                 catch {						
                 }
             }
         }
-        return null
+        return undefined
     }
 
-    const getMergeMethod = (desc: string): MergeMethod => {
-        if (desc.includes("#ROBOMERGE-CONFLICT")) {
+    const getMergeMethod = (desc: string, cl: TrackedCL, sourceCL?: TrackedCL): MergeMethod => {
+        if (compareTrackedCL(cl, data.originalCL)) {
+            return 'initialSubmit'
+        } else if (sourceCL && cl.serverId != sourceCL.serverId) {
+            return 'transfer'
+        } else if (desc.includes("#ROBOMERGE-CONFLICT")) {
             return 'merge_with_conflict'
         } else if (desc.includes("#ROBOMERGE-SOURCE")) {
             return 'automerge'
@@ -61,33 +89,36 @@ export async function trackChange(robo: RoboMerge, logger: ContextualLogger, cl:
         }
     }
 
-    let p4 = PerforceContext.getServerContext(robo.roboMergeLogger)
-    let changes = new Map<number, any>()
+    let changes = new Map<string, TrackedChange>()
 
-    let data: any = {originalCL: cl, changes: {}}
-
-    let gatherCLInfo = async (cl: number, opts?: any) => {
-        let desc = await p4.describe(cl, 1)
-        const mergeMethod = cl != data.originalCL ? getMergeMethod(desc.description) : 'initialSubmit'
+    let gatherCLInfo = async (cl: TrackedCL, opts?: any) => {
+        let p4 = PerforceContext.getServerContext(robo.roboMergeLogger, cl.serverId)
+        let desc = await p4.describe(cl.cl, 1)
+        const mergeMethod = getMergeMethod(desc.description, cl, opts?.sourceCL)
         const clNode = getNode(desc)
         if (clNode || (RobomergeMethodStrings as readonly string[]).includes(mergeMethod)) {
-            changes.set(cl, {desc, node: clNode, sourceCL: opts ? opts.sourceCL : null, destCLs: (opts && opts.lastCL ? [opts.lastCL] : []) })
-            return true
+            changes.set(JSON.stringify(cl), {desc, mergeMethod, node: clNode, sourceCL: opts?.sourceCL, destCLs: (opts && opts.lastCL ? [opts.lastCL] : []) })
         }
-        return false
     }
 
-    await gatherCLInfo(cl)
+    await gatherCLInfo(initialCL)
 
-    const sourceMatch = (changes.get(data.originalCL).desc as DescribeResult).description.match(/#ROBOMERGE-SOURCE: (.*)\n/g)
+    const sourceMatch = (changes.get(JSON.stringify(data.originalCL))?.desc as DescribeResult).description.match(/#ROBOMERGE-SOURCE: (.*)\n/g)
     if (sourceMatch) {
-        let lastCL = cl
+        let lastCL = initialCL
+        data.originalCL = {cl: -1} // we'll find it in a minute, but for the moment clear it so we can get the correct mergeMethod for the initial change
         const matches = Array.from(sourceMatch[0].matchAll(/CL (\d+)/g))
         for (let i=matches.length-1; i>=0;i--) {
-            let sourceCL = parseInt(matches[i][1])
+            let sourceCL = {cl: parseInt(matches[i][1])/*, TODO: deal with cross server #ROBOMERGE-SOURCE tags. serverId: lastCL.serverId */}
             if (i == 0) {
                 data.originalCL = sourceCL
             }
+            if (i == matches.length-1) {
+                // Fix up the initial changes merge method and source now that we know it is not the initial submit
+                let initialChange = changes.get(JSON.stringify(initialCL))!
+                initialChange.mergeMethod = getMergeMethod(initialChange.desc.description, initialCL, sourceCL)
+                initialChange.sourceCL = sourceCL
+            }
             await gatherCLInfo(sourceCL, {lastCL})
             lastCL = sourceCL
         }
@@ -95,16 +126,16 @@ export async function trackChange(robo: RoboMerge, logger: ContextualLogger, cl:
 
     const isFTE = userTags.has("fte")
 
-    let clsToConsider: number[] = [data.originalCL]
-    while (clsToConsider.length > 0) {
-        const clToConsider = clsToConsider[0]
-        clsToConsider = clsToConsider.slice(1)
+    let clsToConsider: TrackedCL[] = [data.originalCL]
+    for (let i = 0; i < clsToConsider.length; i++) {
+        const clToConsider = clsToConsider[i]
 
-        let changeToConsider = changes.get(clToConsider)
+        let changeToConsider = changes.get(JSON.stringify(clToConsider))
         if (!changeToConsider) {
             continue
         }
 
+        const externalServerTargets: Node[] = []
         let includeInResults = false
         let hasAutomergeTarget = false
         let streamDisplayName = ""
@@ -128,8 +159,8 @@ export async function trackChange(robo: RoboMerge, logger: ContextualLogger, cl:
                         hasAutomergeTarget = edge.source.stream == changeToConsider.node.stream
                     }
                 }
-                if (includeInResults && hasAutomergeTarget) {
-                    break
+                if (clToConsider.serverId != edge.target.serverID) {
+                    externalServerTargets.push(edge.target)
                 }
             }
         } else if (isFTE && botFilter.length == 0) {
@@ -152,54 +183,80 @@ export async function trackChange(robo: RoboMerge, logger: ContextualLogger, cl:
             includeInResults = includeInResults && (streamFilter.length == 0 || streamFilter.some(re => upperSteamDisplayName.match(re)))
         }
 
-        // Can't cache this because the passed in changelist is incorrectly labelled initialSubmit the first time
-        // it is considered
-        const mergeMethod = clToConsider != data.originalCL ? getMergeMethod(changeToConsider.desc.description) : 'initialSubmit'
-
+        const p4 = PerforceContext.getServerContext(logger, clToConsider.serverId)
+        let runIntegrated = true
         for (let i=0; i < changeToConsider.desc.entries.length; i++) {
             const entry = changeToConsider.desc.entries[i]
             // integrated for move/delete will point at the paired move/add not where it was merged to in another stream, so not useful to evaluate it
             if (entry.action != 'move/delete') {
-                const integrated = await p4.integrated(null, entry.depotFile, {intoOnly: true, startCL: clToConsider})
-                if (integrated.length > 0) {
-                    for (let integ of integrated) {
-                        const startToRev = integ.startToRev == "#none" ? 0 : parseInt(integ.startToRev.slice(1))
-                        const endToRev = parseInt(integ.endToRev.slice(1))
-                        if (entry.rev <= endToRev && entry.rev > startToRev)
-                        {
-                            const destChange = changes.get(integ.change)
-                            if (destChange) {
-                                destChange.sourceCL = clToConsider
-                            } else {
-                                changeToConsider.destCLs.push(integ.change)
-                                await gatherCLInfo(integ.change, {sourceCL: clToConsider})
+                for (let i=externalServerTargets.length-1; i>=0; i--) {
+                    const externalTarget = externalServerTargets[i]
+                    const externalP4 = PerforceContext.getServerContext(logger, externalTarget.serverID)
+                    const streamName = getStreamFromPath(changeToConsider.node!.stream)
+                    const targetPath = entry.depotFile.replace(streamName || changeToConsider.node!.stream, externalTarget.stream)
+                    const changes = await externalP4.changes(targetPath, 0)
+                    for (let change of changes) {
+                        if (change.desc.includes(clToConsider.cl.toString())) {
+                            const newChangeToConsider = {cl: change.change, serverId: externalTarget.serverID}
+                            changeToConsider.destCLs.push(newChangeToConsider)
+                            await gatherCLInfo(newChangeToConsider, {sourceCL: clToConsider})
+                            externalServerTargets.splice(i, 1)
+                            break
+                        }
+                    }
+                }
+
+                if (runIntegrated) {
+                    const integrated = await p4.integrated(null, entry.depotFile, {intoOnly: true, startCL: clToConsider.cl})
+                    if (integrated.length > 0) {
+                        runIntegrated = false
+                        for (let integ of integrated) {
+                            const startToRev = integ.startToRev == "#none" ? 0 : parseInt(integ.startToRev.slice(1))
+                            const endToRev = parseInt(integ.endToRev.slice(1))
+                            if (entry.rev <= endToRev && entry.rev > startToRev)
+                            {
+                                const destChange = changes.get(integ.change)
+                                if (destChange) {
+                                    destChange.sourceCL = clToConsider
+                                } else {
+                                    const newChangeToConsider = {cl: integ.change, serverId: clToConsider.serverId}
+                                    changeToConsider.destCLs.push(newChangeToConsider)
+                                    await gatherCLInfo(newChangeToConsider, {sourceCL: clToConsider})
+                                }
                             }
                         }
                     }
+                }
+
+                if (!runIntegrated && externalServerTargets.length > 0) {
                     break
                 }
             }
             if (hasAutomergeTarget && changeToConsider.desc.entries.length == 1 && 
-                        (RobomergeMethodStrings as readonly string[]).includes(mergeMethod)) {
+                        (RobomergeMethodStrings as readonly string[]).includes(changeToConsider.mergeMethod)) {
                 // If we only have 1 entry and we didn't get integration info off of it
                 // and the graph suggests we are expecting there could be other changes
                 // get more of the files. If the path is a root directory try and get a sampling 
                 // of files, otherwise just get an additional block to evaluate
                 if (changeToConsider.desc.path.endsWith("/...")) {
-                    const dirs = await p4.dirs(`${changeToConsider.desc.path.slice(0,-3)}*@=${clToConsider}`)
+                    const dirs = await p4.dirs(`${changeToConsider.desc.path.slice(0,-3)}*@=${clToConsider.cl}`)
                     for (const dir of dirs) {
-                        changeToConsider.desc.entries.push(...(await p4.files(`${dir}/...@=${clToConsider}`, 1)))
+                        changeToConsider.desc.entries.push(...(await p4.files(`${dir}/...@=${clToConsider.cl}`, 1)))
                     }
                 }
                 else {
-                    changeToConsider.desc = await p4.describe(clToConsider, 100)
+                    changeToConsider.desc = await p4.describe(clToConsider.cl, 100)
                 }
             }
         }
-        clsToConsider = clsToConsider.concat(changeToConsider.destCLs)
+        for (const destCL of changeToConsider.destCLs) {
+            if (!clsToConsider.find(cl => compareTrackedCL(cl, destCL))) {
+                clsToConsider.push(destCL)
+            }
+        }
 
         if (includeInResults) {
-            data.changes[`${clToConsider}`] = {streamDisplayName, mergeMethod, sourceCL: changeToConsider.sourceCL}
+            data.changes[`${clToConsider.cl}`] = {streamDisplayName, mergeMethod: changeToConsider.mergeMethod, sourceCL: changeToConsider.sourceCL?.cl}
         }
     }
 
