Transfer legend entry

diff --git a/Engine/Extras/RoboMerge/v3/public/trackchange.html b/Engine/Extras/RoboMerge/v3/public/trackchange.html
--- a/Engine/Extras/RoboMerge/v3/public/trackchange.html
+++ b/Engine/Extras/RoboMerge/v3/public/trackchange.html
@@ -61,6 +61,10 @@ Errors attempting to track change <a href="https://p4-swarm.epicgames.net/change
 						<td class="arrow default-flow-arrow"></td>
 						<td class="arrow-label">Populate</td>
 					</tr>
+					<tr title="Change was transfered from one server to another">
+						<td class="arrow gated-auto-merge-arrow"></td>
+						<td class="arrow-label">Transfer</td>
+					</tr>
 				</table>
 
 		</div>
