Fix undefined reference

diff --git a/Engine/Extras/RoboMerge/v3/src/new/graph.ts b/Engine/Extras/RoboMerge/v3/src/new/graph.ts
--- a/Engine/Extras/RoboMerge/v3/src/new/graph.ts
+++ b/Engine/Extras/RoboMerge/v3/src/new/graph.ts
@@ -126,7 +126,7 @@ export class Graph {
 
 	findOrCreateStreamNode(depotPath: string | Branch) {
 		const branchId = makeBranchId(depotPath)
-		const server = (depotPath as Branch).config.streamServer
+		const server = (depotPath as Branch).config?.streamServer
 		return setDefault(this.streamNodes, branchId, new Node(branchId, server))
 	}
 
