It is no longer a transfer error if there is an attempt to transfer a delete of a file that does not exist

diff --git a/Engine/Extras/RoboMerge/v3/src/common/perforce.ts b/Engine/Extras/RoboMerge/v3/src/common/perforce.ts
--- a/Engine/Extras/RoboMerge/v3/src/common/perforce.ts
+++ b/Engine/Extras/RoboMerge/v3/src/common/perforce.ts
@@ -1094,6 +1094,19 @@ export class PerforceContext {
 
 		const FILES_PER_ACTION = 100
 
+		const HandleCaughtError = (reason: any, filterOut?: (s: string) => boolean) => {
+			if (!isExecP4Error(reason)) {
+				failures.push(reason)
+			}
+
+			const [_, output] = reason
+			let reasons = output.split('\n')
+			if (filterOut) {
+				reasons = reasons.filter((s: string) => s.length > 0 && !filterOut(s))
+			}
+			failures.push(...reasons)
+		}
+
 		sortedFiles.forEach((files, action) => {
 			if (action.startsWith('add')) {
 				const fileType = action.slice(4)
@@ -1103,7 +1116,7 @@ export class PerforceContext {
 					promises.push(
 						targetP4.add(roboWorkspace, dest_changelist, targetDepotFiles, fileType, p4quiet)
 						        .then(adds => adds.forEach((add, index) => printCommands.push([`${filesSlice[index][0].depotFile}@${source.changelist}`, add.clientFile])))
-						        .catch(reason => failures.push(reason))
+						        .catch(reason => HandleCaughtError(reason))
 					)
 				}
 			}
@@ -1117,9 +1130,9 @@ export class PerforceContext {
 						        .then(syncs => { 
 									syncs.forEach((sync, index) => printCommands.push([`${filesSlice[index][0].depotFile}@${source.changelist}`, sync.clientFile]))
 									return targetP4.edit(roboWorkspace, dest_changelist, targetDepotFiles, ['-k', '-t', fileType], p4quiet)
-									               .catch(reason => failures.push(reason))
+									               .catch(reason => HandleCaughtError(reason))
 						        })
-						        .catch(reason => failures.push(reason))
+						        .catch(reason => HandleCaughtError(reason))
 					)
 				}
 			}
@@ -1127,7 +1140,7 @@ export class PerforceContext {
 				for (let i=0; i < files.length; i += FILES_PER_ACTION) {
 					promises.push(
 						targetP4.delete(roboWorkspace, dest_changelist, files.slice(i,i+FILES_PER_ACTION).map(([_, targetDepotFile]) => targetDepotFile), p4quiet)
-						.catch(reason => failures.push(reason))
+						.catch(reason => HandleCaughtError(reason, s => s.endsWith("file(s) not on client.")))
 					)
 				}
 			}
@@ -1142,7 +1155,7 @@ export class PerforceContext {
 							.then(() => { return targetP4.edit(roboWorkspace, dest_changelist, targetDepotFile, undefined, p4quiet) })
 							.then(() => { return targetP4.move(roboWorkspace, dest_changelist, targetDepotFile, targetMoveAddDepotFile, ['-f'], p4quiet) })
 							.then(move => printCommands.push([`${integ.fromFile}@${source.changelist}`, move[0].path]))
-							.catch(reason => failures.push(reason))
+							.catch(reason => HandleCaughtError(reason))
 						}
 						else {
 							failures.push(`Unable to look up integrated record for ${file.depotFile} in CL# ${source.changelist}`)
