Deal gracefully when the robomerge origin CL# isn't found

diff --git a/Engine/Extras/RoboMerge/v3/src/robo/ipc.ts b/Engine/Extras/RoboMerge/v3/src/robo/ipc.ts
--- a/Engine/Extras/RoboMerge/v3/src/robo/ipc.ts
+++ b/Engine/Extras/RoboMerge/v3/src/robo/ipc.ts
@@ -188,7 +188,14 @@ export class IPC {
 			userTags.add(tag)
 		}
 	
-		return {...OPERATION_SUCCESS, data: await trackChange(this.robo, this.ipcLogger, {cl}, userTags, queryObj)}
+		const result  = await trackChange(this.robo, this.ipcLogger, {cl}, userTags, queryObj)
+
+		if (result.success) {
+			return {...OPERATION_SUCCESS, data: result.data }
+		}
+		else {
+			return {statusCode: 400, message: result.message!}
+		}
 	}
 
 	private async traceRoute(query: Query, tagsObj?: any): Promise<OperationReturnType> {
diff --git a/Engine/Extras/RoboMerge/v3/src/robo/trackchange.ts b/Engine/Extras/RoboMerge/v3/src/robo/trackchange.ts
--- a/Engine/Extras/RoboMerge/v3/src/robo/trackchange.ts
+++ b/Engine/Extras/RoboMerge/v3/src/robo/trackchange.ts
@@ -1,7 +1,7 @@
 // Copyright Epic Games, Inc. All Rights Reserved.
 
 import { ContextualLogger } from '../common/logger';
-import { DescribeResult, PerforceContext } from '../common/perforce';
+import { DescribeResult, isExecP4Error, PerforceContext } from '../common/perforce';
 import { Node } from '../new/graph';
 import { NodeBotInterface } from './bot-interfaces';
 import { RoboMerge } from './robo';
@@ -93,35 +93,59 @@ export async function trackChange(robo: RoboMerge, logger: ContextualLogger, ini
 
     let gatherCLInfo = async (cl: TrackedCL, opts?: any) => {
         let p4 = PerforceContext.getServerContext(robo.roboMergeLogger, cl.serverId)
-        let desc = await p4.describe(cl.cl, 1)
-        const mergeMethod = getMergeMethod(desc.description, cl, opts?.sourceCL)
-        const clNode = getNode(desc)
-        if (clNode || (RobomergeMethodStrings as readonly string[]).includes(mergeMethod)) {
-            changes.set(JSON.stringify(cl), {desc, mergeMethod, node: clNode, sourceCL: opts?.sourceCL, destCLs: (opts && opts.lastCL ? [opts.lastCL] : []) })
+        try {
+            let desc = await p4.describe(cl.cl, 1)
+            const mergeMethod = getMergeMethod(desc.description, cl, opts?.sourceCL)
+            const clNode = getNode(desc)
+            if (clNode || (RobomergeMethodStrings as readonly string[]).includes(mergeMethod)) {
+                changes.set(JSON.stringify(cl), {desc, mergeMethod, node: clNode, sourceCL: opts?.sourceCL, destCLs: (opts && opts.lastCL ? [opts.lastCL] : []) })
+            }
+        } 
+        catch (reason) {
+            if (!isExecP4Error(reason)) {
+                throw reason
+            }
+
+            let [_, output] = reason
+            // try to be resilient to a missing change
+            if (output.includes("no such changelist.")) {
+                return false
+            }
+            throw reason
         }
+        return true
     }
 
-    await gatherCLInfo(initialCL)
+    if (!(await gatherCLInfo(initialCL))) {
+        return {success: false, message: `Unable to lookup CL# ${initialCL.cl}`}
+    }
 
     const sourceMatch = (changes.get(JSON.stringify(data.originalCL))?.desc as DescribeResult).description.match(/#ROBOMERGE-SOURCE: (.*)\n/g)
     if (sourceMatch) {
+        let initialChange = changes.get(JSON.stringify(initialCL))!
         let lastCL = initialCL
+        let originCL = initialCL
         data.originalCL = {cl: -1} // we'll find it in a minute, but for the moment clear it so we can get the correct mergeMethod for the initial change
         const matches = Array.from(sourceMatch[0].matchAll(/CL (\d+)/g))
         for (let i=matches.length-1; i>=0;i--) {
             let sourceCL = {cl: parseInt(matches[i][1])/*, TODO: deal with cross server #ROBOMERGE-SOURCE tags. serverId: lastCL.serverId */}
-            if (i == 0) {
-                data.originalCL = sourceCL
-            }
-            if (i == matches.length-1) {
-                // Fix up the initial changes merge method and source now that we know it is not the initial submit
-                let initialChange = changes.get(JSON.stringify(initialCL))!
-                initialChange.mergeMethod = getMergeMethod(initialChange.desc.description, initialCL, sourceCL)
-                initialChange.sourceCL = sourceCL
+            if (await gatherCLInfo(sourceCL, {lastCL})) {
+                if (i == 0) {
+                    data.originalCL = sourceCL
+                }
+                if (i == matches.length-1) {
+                    // Fix up the initial changes merge method and source now that we know it is not the initial submit
+                    initialChange.mergeMethod = getMergeMethod(initialChange.desc.description, initialCL, sourceCL)
+                    initialChange.sourceCL = sourceCL
+                }
+                originCL = sourceCL
             }
-            await gatherCLInfo(sourceCL, {lastCL})
             lastCL = sourceCL
         }
+        // For some reason we couldn't get the true origin, put it to the earliest we did gather data for
+        if (data.originalCL.cl == -1) {
+            data.originalCL = originCL
+        }
     }
 
     const isFTE = userTags.has("fte")
@@ -260,5 +284,5 @@ export async function trackChange(robo: RoboMerge, logger: ContextualLogger, ini
         }
     }
 
-    return data
+    return {success: true, data}
 }
\ No newline at end of file
