Fix track change misidentifying the initial submit in some cases

diff --git a/Engine/Extras/RoboMerge/v3/src/robo/trackchange.ts b/Engine/Extras/RoboMerge/v3/src/robo/trackchange.ts
--- a/Engine/Extras/RoboMerge/v3/src/robo/trackchange.ts
+++ b/Engine/Extras/RoboMerge/v3/src/robo/trackchange.ts
@@ -39,6 +39,12 @@ interface TrackChangeResults {
     changes: any
 }
 
+interface GatherOptions {
+    sourceCL?: TrackedCL,
+    lastCL?: TrackedCL,
+    initialSubmit?: boolean
+}
+
 export async function trackChange(robo: RoboMerge, logger: ContextualLogger, initialCL: TrackedCL, userTags: Set<string>, filters: TrackFilters) {
 
     const streamFilter = (() => {
@@ -73,15 +79,15 @@ export async function trackChange(robo: RoboMerge, logger: ContextualLogger, ini
         return undefined
     }
 
-    const getMergeMethod = (desc: string, cl: TrackedCL, sourceCL?: TrackedCL): MergeMethod => {
-        if (compareTrackedCL(cl, data.originalCL)) {
-            return 'initialSubmit'
-        } else if (sourceCL && cl.serverId != sourceCL.serverId) {
+    const getMergeMethod = (desc: string, cl: TrackedCL, opts?: GatherOptions): MergeMethod => {
+        if (opts?.sourceCL && cl.serverId != opts.sourceCL.serverId) {
             return 'transfer'
         } else if (desc.includes("#ROBOMERGE-CONFLICT")) {
             return 'merge_with_conflict'
         } else if (desc.includes("#ROBOMERGE-SOURCE")) {
             return 'automerge'
+        } else if (opts?.initialSubmit || compareTrackedCL(cl, data.originalCL)) {
+            return 'initialSubmit'
         } else if (desc.includes("Populate")) {
             return 'populate'
         } else {
@@ -91,11 +97,11 @@ export async function trackChange(robo: RoboMerge, logger: ContextualLogger, ini
 
     let changes = new Map<string, TrackedChange>()
 
-    let gatherCLInfo = async (cl: TrackedCL, opts?: any) => {
+    let gatherCLInfo = async (cl: TrackedCL, opts?: GatherOptions) => {
         let p4 = PerforceContext.getServerContext(robo.roboMergeLogger, cl.serverId)
         try {
             let desc = await p4.describe(cl.cl, 1)
-            const mergeMethod = getMergeMethod(desc.description, cl, opts?.sourceCL)
+            const mergeMethod = getMergeMethod(desc.description, cl, opts)
             const clNode = getNode(desc)
             if (clNode || (RobomergeMethodStrings as readonly string[]).includes(mergeMethod)) {
                 changes.set(JSON.stringify(cl), {desc, mergeMethod, node: clNode, sourceCL: opts?.sourceCL, destCLs: (opts && opts.lastCL ? [opts.lastCL] : []) })
@@ -122,27 +128,23 @@ export async function trackChange(robo: RoboMerge, logger: ContextualLogger, ini
 
     const sourceMatch = (changes.get(JSON.stringify(data.originalCL))?.desc as DescribeResult).description.match(/#ROBOMERGE-SOURCE: (.*)\n/g)
     if (sourceMatch) {
-        let initialChange = changes.get(JSON.stringify(initialCL))!
         let lastCL = initialCL
         let originCL = initialCL
-        data.originalCL = {cl: -1} // we'll find it in a minute, but for the moment clear it so we can get the correct mergeMethod for the initial change
         const matches = Array.from(sourceMatch[0].matchAll(/CL (\d+)/g))
         for (let i=matches.length-1; i>=0;i--) {
             let sourceCL = {cl: parseInt(matches[i][1])/*, TODO: deal with cross server #ROBOMERGE-SOURCE tags. serverId: lastCL.serverId */}
-            if (await gatherCLInfo(sourceCL, {lastCL})) {
-                if (i == 0) {
-                    data.originalCL = sourceCL
-                }
+            if (await gatherCLInfo(sourceCL, {lastCL, initialSubmit: i==0})) {
                 if (i == matches.length-1) {
+                    let initialChange = changes.get(JSON.stringify(initialCL))!
                     // Fix up the initial changes merge method and source now that we know it is not the initial submit
-                    initialChange.mergeMethod = getMergeMethod(initialChange.desc.description, initialCL, sourceCL)
+                    initialChange.mergeMethod = getMergeMethod(initialChange.desc.description, initialCL, {sourceCL})
                     initialChange.sourceCL = sourceCL
                 }
                 originCL = sourceCL
             }
             lastCL = sourceCL
         }
-        // For some reason we couldn't get the true origin, put it to the earliest we did gather data for
+        // We've now determined the earliest CL we are going to start tracking from
         if (data.originalCL.cl == -1) {
             data.originalCL = originCL
         }
