More fixes to tracking back the initial hierarchy of changes

diff --git a/Engine/Extras/RoboMerge/v3/src/robo/trackchange.ts b/Engine/Extras/RoboMerge/v3/src/robo/trackchange.ts
--- a/Engine/Extras/RoboMerge/v3/src/robo/trackchange.ts
+++ b/Engine/Extras/RoboMerge/v3/src/robo/trackchange.ts
@@ -132,8 +132,9 @@ export async function trackChange(robo: RoboMerge, logger: ContextualLogger, ini
         let originCL = initialCL
         const matches = Array.from(sourceMatch[0].matchAll(/CL (\d+)/g))
         for (let i=matches.length-1; i>=0;i--) {
-            let sourceCL = {cl: parseInt(matches[i][1])/*, TODO: deal with cross server #ROBOMERGE-SOURCE tags. serverId: lastCL.serverId */}
-            if (await gatherCLInfo(sourceCL, {lastCL, initialSubmit: i==0})) {
+            const sourceCL = {cl: parseInt(matches[i][1])/*, TODO: deal with cross server #ROBOMERGE-SOURCE tags. serverId: lastCL.serverId */}
+            const parentCL = i > 0 ? {cl: parseInt(matches[i-1][1])} : undefined
+            if (await gatherCLInfo(sourceCL, {sourceCL: parentCL, lastCL, initialSubmit: i==0})) {
                 if (i == matches.length-1) {
                     let initialChange = changes.get(JSON.stringify(initialCL))!
                     // Fix up the initial changes merge method and source now that we know it is not the initial submit
@@ -145,9 +146,7 @@ export async function trackChange(robo: RoboMerge, logger: ContextualLogger, ini
             lastCL = sourceCL
         }
         // We've now determined the earliest CL we are going to start tracking from
-        if (data.originalCL.cl == -1) {
-            data.originalCL = originCL
-        }
+        data.originalCL = originCL
     }
 
     const isFTE = userTags.has("fte")
