Sort changes (which would be from another server) after their source rather than it being strictly CL# based

diff --git a/Engine/Extras/RoboMerge/v3/public/js/trackchange.js b/Engine/Extras/RoboMerge/v3/public/js/trackchange.js
--- a/Engine/Extras/RoboMerge/v3/public/js/trackchange.js
+++ b/Engine/Extras/RoboMerge/v3/public/js/trackchange.js
@@ -60,9 +60,39 @@ async function generateChangeList(dataObj) {
 
 	const gitHubCommits = await getGitHubCommits(dataObj.changes)
 	let html = '<div style="margin: auto; width: 80%;"><table class="table"><tbody>'
-	for (const cl in dataObj.changes) {
+	// By default the keys are going to be numerically ordered, but we want changes to appear after their parent
+	// So go through and build them in a list order we prefer
+	const orderedChanges = Object.keys(dataObj.changes).map(k => parseInt(k))
+	const placedChanges = new Set()
+	for (let i = 0; i < orderedChanges.length; i++) {
+		const sourceCL = dataObj.changes[orderedChanges[i].toString()].sourceCL
+		if (sourceCL && !placedChanges.has(sourceCL)) {
+			let insertPoint = undefined
+			for (let j = i+1; j < orderedChanges.length; j++) {
+				if (orderedChanges[j] == sourceCL) {
+					insertPoint = j
+				}
+				else if (insertPoint) {
+					// Keep the children with lower CL# of a given stream in numerical order
+					if (orderedChanges[j] < orderedChanges[i]) {
+						insertPoint = j
+					} else {
+						break
+					}
+				}
+			}
+			if (insertPoint) {
+				const removedKey = orderedChanges.splice(i, 1)[0]
+				orderedChanges.splice(insertPoint, 0, removedKey)
+				i--
+				continue
+			}
+		}
+		placedChanges.add(orderedChanges[i])
+	}
+	for (const cl of orderedChanges) {
 		html += '<tr valign="middle">'
-		html += `<td><b>${dataObj.changes[cl].streamDisplayName}</b></td>`
+		html += `<td><b>${dataObj.changes[cl.toString()].streamDisplayName}</b></td>`
 		html += `<td><a href="https://p4-swarm.epicgames.net/changes/${cl}" target="_blank">CL#${cl}</a></td>`
 		html += `${formatGitHubCommit(gitHubCommits.get(cl))}`
 		html += '</tr>'
