Dynamically identify the swarm URL for a perforce server as part of initialization
Track Changes and Preview link to swarm CL on remote servers

diff --git a/Engine/Extras/RoboMerge/v3/public/js/preview.js b/Engine/Extras/RoboMerge/v3/public/js/preview.js
--- a/Engine/Extras/RoboMerge/v3/public/js/preview.js
+++ b/Engine/Extras/RoboMerge/v3/public/js/preview.js
@@ -139,7 +139,6 @@ function reloadData()
 	window.location.href = url;
 }
 
-
 function doit(query)
 {
 	const clgroups = query.match(/cl=([^&]*)/);
@@ -188,6 +187,14 @@ function doit(query)
 				false)
 			);
 			$('.bots').html([...botNames].map(s => `<tt>${s.toLowerCase()}</tt>`).join(', '))
+			const swarmURL = data.version.match(/preview:(.*)/)[1]
+			const cl = swarmURL.match(/.*\/(\d+)/)[1]
+			if (swarmURL.startsWith("http")) {
+				$('.swarmURL').html(`<a href="${swarmURL}" target="_blank">CL# ${cl}</a>`)
+			}
+			else {
+				$('.swarmURL').html(`CL# ${cl}`)
+			}
 			if (data.validationWarnings && data.validationWarnings.length > 0) {
 				const $warningPanel = $('#warning-panel');
 				const warningMsg = "Warning:" + data.validationWarnings.join('\n').replace(/\t/g, '    ')
diff --git a/Engine/Extras/RoboMerge/v3/public/js/trackchange.js b/Engine/Extras/RoboMerge/v3/public/js/trackchange.js
--- a/Engine/Extras/RoboMerge/v3/public/js/trackchange.js
+++ b/Engine/Extras/RoboMerge/v3/public/js/trackchange.js
@@ -91,9 +91,14 @@ async function generateChangeList(dataObj) {
 		placedChanges.add(orderedChanges[i])
 	}
 	for (const cl of orderedChanges) {
+		const changeDetails = dataObj.changes[cl.toString()]
 		html += '<tr valign="middle">'
-		html += `<td><b>${dataObj.changes[cl.toString()].streamDisplayName}</b></td>`
-		html += `<td><a href="https://p4-swarm.epicgames.net/changes/${cl}" target="_blank">CL#${cl}</a></td>`
+		html += `<td><b>${changeDetails.streamDisplayName}</b></td>`
+		if (changeDetails.swarmLink) {
+			html += `<td><a href="${changeDetails.swarmLink}" target="_blank">CL#${cl}</a></td>`
+		} else {
+			html += `<td>CL#${cl}</td>`
+		}
 		html += `${formatGitHubCommit(gitHubCommits.get(cl))}`
 		html += '</tr>'
 	}
@@ -110,16 +115,19 @@ function createGraph(changes) {
 		'node [shape=box, style=filled, fontname="sans-serif", fillcolor=moccasin];'
 	]
 	for (let change in changes) {
-		const attrs = [
-			['label', `<<b>${getStreamGraphName(changes[change].streamDisplayName)}</b><br/>${change}>`],
-			['tooltip', `"https://p4-swarm.epicgames.net/changes/${change}"`],
-			['URL', `"https://p4-swarm.epicgames.net/changes/${change}"`],
+		const changeDetails = changes[change]
+		let attrs = [
+			['label', `<<b>${getStreamGraphName(changeDetails.streamDisplayName)}</b><br/>${change}>`],
+			['tooltip', `"${changeDetails.swarmLink}"` || `"CL# ${change}"`],
 			['target', `"_blank"`],
 			['margin', `"0.5,0.1"`],
 			['fontsize', 15],
-		];
-		const attrStrs = attrs.map(([key, value]) => `${key}=${value}`);
-		lines.push(`_${change} [${attrStrs.join(', ')}];`);		
+		]
+		if (changeDetails.swarmLink) {
+			attrs.push(['URL', `"${changeDetails.swarmLink}"`])
+		}
+		const attrStrs = attrs.map(([key, value]) => `${key}=${value}`)
+		lines.push(`_${change} [${attrStrs.join(', ')}];`);	
 	}
 	for (let change in changes) {
 		if (changes[change].sourceCL in changes) {
@@ -163,6 +171,9 @@ async function buildResults(data) {
 	} else {
 		const $errorPanel = $('#error-panel');
 		$('pre', $errorPanel).html('No results found.')
+		if (data.data.swarmURL) {
+			$('.swarmURL').html(`<a href="${data.data.swarmURL}/changes/${data.originalCL.cl}" target="_blank">CL# ${data.originalCL.cl}</a>`)
+		}
 		$errorPanel.show();
 	}
 	receivedTrackingResults()
diff --git a/Engine/Extras/RoboMerge/v3/public/preview.html b/Engine/Extras/RoboMerge/v3/public/preview.html
--- a/Engine/Extras/RoboMerge/v3/public/preview.html
+++ b/Engine/Extras/RoboMerge/v3/public/preview.html
@@ -32,15 +32,15 @@
 	<div class="container-fluid">
 		<div class="row m-1">
 			<div class="col" id="success-panel" style="display:none;">
-				<div class="alert alert-success" role="alert">Graph for <span class="bots"></span> from <a href="https://p4-swarm.epicgames.net/changes/{{ cl }}" target="_blank">CL#{{ cl }}</a></div>
+				<div class="alert alert-success" role="alert">Graph for <span class="bots"></span> from <span class="swarmURL"></span></div>
 			</div>
 			<div class="col">
 				<div id="warning-panel" style="display:none;"> <!-- show or hide in script -->
-					<div class="alert alert-warning" role="alert">Warnings while previewing <span class="bots"></span> from <a href="https://p4-swarm.epicgames.net/changes/{{ cl }}" target="_blank">CL#{{ cl }}</a></div>
+					<div class="alert alert-warning" role="alert">Warnings while previewing <span class="bots"></span> from <span class="swarmURL"></span></div>
 					<pre></pre>
 				</div>
 				<div id="error-panel" style="display:none;"> <!-- show or hide in script -->
-					<div class="alert alert-warning" role="alert">Errors attempting to preview <span class="bots"></span> from <a href="https://p4-swarm.epicgames.net/changes/{{ cl }}" target="_blank">CL#{{ cl }}</a></div>
+					<div class="alert alert-warning" role="alert">Errors attempting to preview <span class="bots"></span> from <span class="swarmURL"></span></div>
 					<pre></pre>
 				</div>
 			</div>
diff --git a/Engine/Extras/RoboMerge/v3/public/trackchange.html b/Engine/Extras/RoboMerge/v3/public/trackchange.html
--- a/Engine/Extras/RoboMerge/v3/public/trackchange.html
+++ b/Engine/Extras/RoboMerge/v3/public/trackchange.html
@@ -35,7 +35,7 @@
 		
 		<div id="error-panel" style="display:none;"> <!-- show or hide in script -->
 			<div class="alert alert-warning" role="alert">
-Errors attempting to track change <a href="https://p4-swarm.epicgames.net/changes/{{ cl }}" target="_blank">CL#{{ cl }}</a>
+Errors attempting to track change <span class="swarmURL">CL# {{ cl }}</span>
 			</div>
 			<pre></pre>
 		</div>
diff --git a/Engine/Extras/RoboMerge/v3/src/common/perforce.ts b/Engine/Extras/RoboMerge/v3/src/common/perforce.ts
--- a/Engine/Extras/RoboMerge/v3/src/common/perforce.ts
+++ b/Engine/Extras/RoboMerge/v3/src/common/perforce.ts
@@ -331,6 +331,7 @@ interface ServerConfig {
 	serverAddress?: string
 	serverVersion: number
 	multiServerEnvironment: boolean
+	swarmURL?: string
 }
 
 export class ResolveResult {
@@ -522,7 +523,13 @@ export class PerforceContext {
 		let resp = output[0];
 	
 		if (resp && resp.User) {
-			const rawServerVersion = await PerforceContext._execP4(logger, null, ["-u", resp.User, "-ztag","-F","%serverVersion%","info"], { serverAddress} )
+
+			const [rawServerVersion, serversOutput, swarmOutput] = await Promise.all([
+				PerforceContext._execP4(logger, null, ["-u", resp.User, "-ztag","-F","%serverVersion%","info"], { serverAddress} ),
+				PerforceContext.execAndParse(logger, null, ["-u", resp.User, "servers"], { serverAddress }),
+				PerforceContext._execP4(logger, null, ["-u", resp.User, "-ztag", "-F", "%value%", "property", "-l", "-n", "P4.Swarm.URL"], { serverAddress })
+			])
+
 			const match = rawServerVersion.match(/.*\/(\d+\.\d+)\/.*/)
 			if (!match) {
 				throw new Error(`Unable to parse server version from ${rawServerVersion}`)
@@ -532,13 +539,12 @@ export class PerforceContext {
 				throw new Error(`Robomerge requires a minimum server version of ${MIN_SERVER_VERSION}`)
 			}
 	
-			const serversOutput = await PerforceContext.execAndParse(logger, null, ["-u", resp.User, "servers"], { serverAddress });
-
 			let serverConfig: ServerConfig = {
 				username: resp.User,
 				serverAddress,
 				serverVersion,
-				multiServerEnvironment: serversOutput.length > 1
+				multiServerEnvironment: serversOutput.length > 1,
+				swarmURL: swarmOutput.length > 0 ? swarmOutput.trim() : undefined
 			}
 			this.servers.set(serverID, serverConfig)
 
@@ -606,6 +612,10 @@ export class PerforceContext {
 		return this._serverConfig.multiServerEnvironment
 	}
 
+	get swarmURL() {
+		return this._serverConfig.swarmURL
+	}
+
 	// get a list of all pending changes for this user
 	async get_pending_changes(workspaceName?: string, edgeServerAddress?: string) {
 		if (!this.username) {
diff --git a/Engine/Extras/RoboMerge/v3/src/robo/preview.ts b/Engine/Extras/RoboMerge/v3/src/robo/preview.ts
--- a/Engine/Extras/RoboMerge/v3/src/robo/preview.ts
+++ b/Engine/Extras/RoboMerge/v3/src/robo/preview.ts
@@ -15,7 +15,7 @@ export async function getPreview(cl: number, singleBot?: string) {
 
 	let p4 = PerforceContext.getServerContext(logger)
 
-	const status = new Status(new Date(), 'preview', logger)
+	const status = new Status(new Date(), `preview:${p4.swarmURL}/changes/${cl}`, logger)
 
 	const bots: [string, string][] = []
 	for (const entry of (await p4.describe(cl, undefined, true)).entries) {
diff --git a/Engine/Extras/RoboMerge/v3/src/robo/trackchange.ts b/Engine/Extras/RoboMerge/v3/src/robo/trackchange.ts
--- a/Engine/Extras/RoboMerge/v3/src/robo/trackchange.ts
+++ b/Engine/Extras/RoboMerge/v3/src/robo/trackchange.ts
@@ -36,6 +36,7 @@ interface TrackedChange {
 
 interface TrackChangeResults {
     originalCL: TrackedCL,
+    swarmURL?: string,
     changes: any
 }
 
@@ -58,7 +59,7 @@ export async function trackChange(robo: RoboMerge, logger: ContextualLogger, ini
     })()
 
     const graph = robo.graph.graph
-    let data: TrackChangeResults = {originalCL: initialCL, changes: {}}
+    let data: TrackChangeResults = {originalCL: initialCL, swarmURL: PerforceContext.getServerContext(robo.roboMergeLogger).swarmURL, changes: {}}
 
     const getStreamFromPath = (path: string) => {
         const match = path.match(/\/\/[^\/]*\/[^\/]*/)
@@ -281,7 +282,8 @@ export async function trackChange(robo: RoboMerge, logger: ContextualLogger, ini
         }
 
         if (includeInResults) {
-            data.changes[`${clToConsider.cl}`] = {streamDisplayName, mergeMethod: changeToConsider.mergeMethod, sourceCL: changeToConsider.sourceCL?.cl}
+            const swarmLink = p4.swarmURL ? `${p4.swarmURL}/changes/${clToConsider.cl}` : undefined
+            data.changes[`${clToConsider.cl}`] = {streamDisplayName, swarmLink, mergeMethod: changeToConsider.mergeMethod, sourceCL: changeToConsider.sourceCL?.cl}
         }
     }
 
