Use -f when adding so that files with problematic characters can be transfered

diff --git a/Engine/Extras/RoboMerge/v3/src/common/perforce.ts b/Engine/Extras/RoboMerge/v3/src/common/perforce.ts
--- a/Engine/Extras/RoboMerge/v3/src/common/perforce.ts
+++ b/Engine/Extras/RoboMerge/v3/src/common/perforce.ts
@@ -1753,7 +1753,8 @@ export class PerforceContext {
 			filePath = [filePath]
 		}
 		const args = [
-			'add', 
+			'add',
+			'-f',
 			'-c', cl.toString(), 
 			...(filetype ? ['-t', filetype] : []),
 			...filePath
