Remove unused code

diff --git a/Engine/Extras/RoboMerge/v3/src/robo/status.ts b/Engine/Extras/RoboMerge/v3/src/robo/status.ts
--- a/Engine/Extras/RoboMerge/v3/src/robo/status.ts
+++ b/Engine/Extras/RoboMerge/v3/src/robo/status.ts
@@ -159,9 +159,5 @@ export function preprocess(data: UserStatusData) {
 				node.retry_cl = node.lastBlockage
 			}
 		}
-
-
-		// look for unblocked edges with a current conflict
-		// conflicts.find(c => c.cl === edge.last_cl)
 	}
 }
