Remove unused code

diff --git a/Engine/Extras/RoboMerge/v3/public/js/unlock.js b/Engine/Extras/RoboMerge/v3/public/js/unlock.js
--- a/Engine/Extras/RoboMerge/v3/public/js/unlock.js
+++ b/Engine/Extras/RoboMerge/v3/public/js/unlock.js
@@ -218,10 +218,6 @@ function visualizeAuthor(author, files) {
 }
 
 function visualizeUnlockVerification(requestedBranchCl, unlockJson) {
-    // Debug
-    const prettyJson = JSON.stringify(unlockJson, null, 2)
-    let linkToSwarm = makeClLink(requestedBranchCl, `CL ${requestedBranchCl}`)
-
     let authorDict = {}
     for (const file of unlockJson.files) {
 
