Use dynamically determined swarm URL for stomp, unlock, conflict, and last changelists

diff --git a/Engine/Extras/RoboMerge/v3/public/js/boilerplate.js b/Engine/Extras/RoboMerge/v3/public/js/boilerplate.js
--- a/Engine/Extras/RoboMerge/v3/public/js/boilerplate.js
+++ b/Engine/Extras/RoboMerge/v3/public/js/boilerplate.js
@@ -81,17 +81,27 @@ function removeOnBeforeUnload() {
 	$( window ).off("beforeunload")
 }
 
-function makeClLink(cl, alias) {
+function makeClLink(cl, swarmURL, alias) {
 	if (typeof(cl) !== "number") {
 		return `<span style="font-weight:bolder">${cl}</span>`
 	}
-	return `<a href="https://p4-swarm.companyname.net/changes/${cl}" target="_blank">${alias ? alias : cl}</a>`
+	if (swarmURL) {
+		return `<a href="${swarmURL}/changes/${cl}" target="_blank">${alias || cl}</a>`
+	}
+	else {
+		return `${alias || cl}`
+	}
 }
 
-function makeSwarmFileLink(depotPath, alias) {
-	return `<a href="https://p4-swarm.companyname.net/files/${
-		encodeURI(depotPath.endsWith('/') ? depotPath.slice(2, -1) : depotPath.slice(2)) // Remove '//' (and any trailing slash) to ensure the URL is valid
-	}#commits" target="_blank">${alias ? alias : depotPath}</a>`
+function makeSwarmFileLink(depotPath, swarmURL, alias) {
+	if (swarmURL) {
+		return `<a href="${swarmURL}/files/${
+			encodeURI(depotPath.endsWith('/') ? depotPath.slice(2, -1) : depotPath.slice(2)) // Remove '//' (and any trailing slash) to ensure the URL is valid
+		}#commits" target="_blank">${alias ? alias : depotPath}</a>`
+	}
+	else {
+		return `${alias ? alias : depotPath}`
+	}
 }
 
 function debug(message) {
@@ -896,7 +906,7 @@ function createNodeRow(nodeData, includeActions, includeCollapseControls=true) {
 
 	// Last Change Column 
 	if (window.showNodeLastChanges) {
-		columnArray.push(renderLastChangeCell_Common(nodeData.bot, nodeData.def.name, nodeData.last_cl, nodeAPIOp, operationArgs))
+		columnArray.push(renderLastChangeCell_Common(nodeData.bot, nodeData.def.name, nodeData.last_cl, nodeData.swarmURL, nodeAPIOp, operationArgs))
 	}
 	else {
 		columnArray.push($('<div>').addClass('lastchangecell'))
@@ -944,8 +954,8 @@ function createEdgeRow(nodeData, edgeData, includeActions) {
 	if (edgeData.num_changes_remaining > 1) {
 		catchupText = `pending: ${edgeData.num_changes_remaining}`
 	}
-	columnArray.push(renderLastChangeCell_Common(nodeData.bot, nodeData.def.name, edgeData.last_cl, edgeAPIOp,
-														operationArgs, catchupText, edgeData.display_name))
+	columnArray.push(renderLastChangeCell_Common(nodeData.bot, nodeData.def.name, edgeData.last_cl, edgeData.swarmURL,
+													edgeAPIOp,operationArgs, catchupText, edgeData.display_name))
 	postRenderLastChangeCell_Edge(nodeData.bot, columnArray[columnArray.length - 1], edgeData, edgeAPIOp, operationArgs)
 
 	return columnArray
@@ -1079,7 +1089,7 @@ function createPauseDivs(data, conflict) {
 
 		divs.push($('<div class="info-block conflict">')
 			.append(
-				$('<p>').html('<h6>Blocked on change ' + makeClLink(changelist) + `</h6> ${info}`)
+				$('<p>').html('<h6>Blocked on change ' + makeClLink(changelist, data.swarmURL) + `</h6> ${info}`)
 			)
 		)
 	}
@@ -1682,10 +1692,10 @@ function renderActionsCell_Edge(actionCell, nodeData, edgeData, conflict=null) {
 }
 
 // Helper function to create last change cell
-function renderLastChangeCell_Common(botname, nodename, lastCl, operationFunction, operationArgs, catchupText=null, edgename=null) {
+function renderLastChangeCell_Common(botname, nodename, lastCl, swarmURL, operationFunction, operationArgs, catchupText=null, edgename=null) {
 	const lastChangeCell = $('<div>').addClass('lastchangecell')
 
-	const swarmLink = $(makeClLink(lastCl)).appendTo(lastChangeCell)
+	const swarmLink = $(makeClLink(lastCl, swarmURL)).appendTo(lastChangeCell)
 	if (catchupText) {
 		lastChangeCell.append($('<div>').addClass('catchup').text(catchupText))
 	}
diff --git a/Engine/Extras/RoboMerge/v3/public/js/stomp.js b/Engine/Extras/RoboMerge/v3/public/js/stomp.js
--- a/Engine/Extras/RoboMerge/v3/public/js/stomp.js
+++ b/Engine/Extras/RoboMerge/v3/public/js/stomp.js
@@ -197,12 +197,12 @@ function transitionDisplayStompResults() {
 }
 
 // Create visualization of the changelist, with links to swarm, a display of the description and a list of relevant files to this request
-function visualizeChangelist(changelist, changelistData) {
+function visualizeChangelist(changelist, changelistData, swarmURL) {
     const changelistDiv = $(`<div id="${changelist}" class="stomp-visual">`)
     if (robomergeUser && changelistData.author.toLowerCase() === robomergeUser.userName.toLowerCase()) {
-        changelistDiv.append($('<h3>').html(`${makeClLink(changelist)} by <strong>you</strong>`))
+        changelistDiv.append($('<h3>').html(`${makeClLink(changelist, swarmURL)} by <strong>you</strong>`))
     } else {
-        changelistDiv.append($('<h3>').html(`${makeClLink(changelist)} by ${changelistData.author}`))
+        changelistDiv.append($('<h3>').html(`${makeClLink(changelist, swarmURL)} by ${changelistData.author}`))
     }
     changelistDiv.append($('<hr style="margin-top:0em; border-width: 2px;">'))
 
@@ -234,38 +234,33 @@ function visualizeChangelist(changelist, changelistData) {
 }
 
 function visualizeStompVerification(requestedBranchCl, stompJson) {
-    // Debug
-    const prettyJson = JSON.stringify(stompJson, null, 2)
-    console.log(`Stomp Verification Debug JSON:\n${prettyJson}`)
-    let linkToSwarm = makeClLink(requestedBranchCl, `CL ${requestedBranchCl}`)
-
     // For visualization, sort data by stomped revision
     // (JSON data is sorted by file)
     let clDict = {}
     let encounteredStompedRevisionsCalculationIssues = false
     for (const file of stompJson.files) {
-        const linkToSwarm = makeSwarmFileLink(file.targetFileName)
+        const swarmFileLink = makeSwarmFileLink(file.targetFileName, stompJson.swarmURL)
 
         // Process resolved files
         if (file.resolved) {
             // Add text files to the warning list
             if (file.filetype.startsWith('text')) {
                 $('#mixedMergeWarningList').append(
-                    $('<li>').html(linkToSwarm)
+                    $('<li>').html(swarmFileLink)
                 )
             }
         } 
         // Add unresolved text files to the error list
         else if (file.filetype.startsWith('text')) {
             $('#unresolvedTextFilesList').append(
-                $('<li>').html(linkToSwarm)
+                $('<li>').html(swarmFileLink)
             )
         }
 
         if (file.stompedRevisionsSkipped) {
             $('#stompedRevisionsSkippedDiv').removeClass('initiallyHidden')
             $('#stompedRevisionsSkippedList').append(
-                $('<li>').html(linkToSwarm)
+                $('<li>').html(swarmFileLink)
             )
             // If this file skipped stomped revisions, don't bother attempting to add to clDict
             continue
@@ -274,7 +269,7 @@ function visualizeStompVerification(requestedBranchCl, stompJson) {
         if (file.stompedRevisionsCalculationIssues) {
             encounteredStompedRevisionsCalculationIssues = true
             $('#stompedRevisionsIssuesList').append(
-                $('<li>').html(linkToSwarm)
+                $('<li>').html(swarmFileLink)
             )
             // If this file couldn't calculate stomped revisions, don't bother attempting to add to clDict
             continue
@@ -302,6 +297,8 @@ function visualizeStompVerification(requestedBranchCl, stompJson) {
         }
     }
 
+    const linkToSwarm = makeClLink(requestedBranchCl, stompJson.swarmURL, `CL ${requestedBranchCl}`)
+
     if (encounteredStompedRevisionsCalculationIssues) {
         $('#stompedRevisionsIssuesDiv').removeClass('initiallyHidden')
     }
@@ -331,7 +328,7 @@ function visualizeStompVerification(requestedBranchCl, stompJson) {
     // Now display each changelist
     const sortedKeys = Object.keys(clDict).sort()
     sortedKeys.forEach(function(key) {
-        $('#resultVisualization').append(visualizeChangelist(key, clDict[key]))
+        $('#resultVisualization').append(visualizeChangelist(key, clDict[key], stompJson.swarmURL))
     })
 }
 
diff --git a/Engine/Extras/RoboMerge/v3/src/robo/branch-interfaces.ts b/Engine/Extras/RoboMerge/v3/src/robo/branch-interfaces.ts
--- a/Engine/Extras/RoboMerge/v3/src/robo/branch-interfaces.ts
+++ b/Engine/Extras/RoboMerge/v3/src/robo/branch-interfaces.ts
@@ -52,6 +52,9 @@ export interface StompVerification extends OperationResult {
 	nonBinaryFilesResolved?: boolean // Warning to the user that non-binary files are included in changelist, but passed the first resolve
 	svFiles?: StompVerificationFile[] // Array of remaining files
 
+	// Swarm URL, assumes source and target are on same server as we don't support cross-server stomp
+	swarmURL?: string
+
 	// Convenience booleans to alert user to problems in verification result
 	remainingAllBinary?: boolean // Check to see if unresolved non-binary files are remaining -- we shouldn't stomp those!
 }
diff --git a/Engine/Extras/RoboMerge/v3/src/robo/edgebot.ts b/Engine/Extras/RoboMerge/v3/src/robo/edgebot.ts
--- a/Engine/Extras/RoboMerge/v3/src/robo/edgebot.ts
+++ b/Engine/Extras/RoboMerge/v3/src/robo/edgebot.ts
@@ -1097,6 +1097,7 @@ class EdgeBotImpl extends PerforceStatefulBot {
 		status.rootPath = this.targetBranch.rootPath
 
 		status.last_cl = this.lastCl
+		status.swarmURL = this.p4.swarmURL
 		
 		status.is_active = this.isActive
 		status.is_available = this.isAvailable
diff --git a/Engine/Extras/RoboMerge/v3/src/robo/ipc.ts b/Engine/Extras/RoboMerge/v3/src/robo/ipc.ts
--- a/Engine/Extras/RoboMerge/v3/src/robo/ipc.ts
+++ b/Engine/Extras/RoboMerge/v3/src/robo/ipc.ts
@@ -538,7 +538,9 @@ export class IPC {
 						nonBinaryFilesResolved: stompVerification.nonBinaryFilesResolved,
 						remainingAllBinary: stompVerification.remainingAllBinary,
 						files: stompVerification.svFiles,
-						validRequest: stompVerification.validRequest
+						validRequest: stompVerification.validRequest,
+						sourceSwarmURL: stompVerification.sourceSwarmURL,
+						targetSwarmURL: stompVerification.targetSwarmURL
 					} )
 				}
 			}
diff --git a/Engine/Extras/RoboMerge/v3/src/robo/nodebot.ts b/Engine/Extras/RoboMerge/v3/src/robo/nodebot.ts
--- a/Engine/Extras/RoboMerge/v3/src/robo/nodebot.ts
+++ b/Engine/Extras/RoboMerge/v3/src/robo/nodebot.ts
@@ -638,25 +638,32 @@ export class NodeBot extends PerforceStatefulBot implements NodeBotInterface {
 	 */
 	async verifyStomp(changeCl: number, targetBranchName: string) : Promise<StompVerification> {
 		this._log_action(`Verifying stomp request of CL ${changeCl} to ${targetBranchName}`, "verbose")
+
+		const failureResult: StompVerification = {
+			success: false,
+			message: "TBD",
+			swarmURL: this.p4.swarmURL
+		}
+
 		if (!changeCl || !targetBranchName) {
-			return { success: false, message: "Missing parameters for stomp node operation" }
+			return { ...failureResult, message: "Missing parameters for stomp node operation" }
 		}
 
 		const targetBranch = this._getBranch(targetBranchName)
 
 		if (!targetBranch) {
-			return { success: false, message: `Unable to retrieve branch infomation for "${targetBranchName}". Unable to stomp, please contact Robomerge support."` }
+			return { ...failureResult, message: `Unable to retrieve branch infomation for "${targetBranchName}". Unable to stomp, please contact Robomerge support."` }
 		}
 
 		const edge = this.getImmediateEdge(targetBranch)
 
 		if (!edge) {
-			return { success: false, message: `This bot has no edge for specified branch "${targetBranch.name}" on ${this.fullName}`}
+			return { ...failureResult, message: `This bot has no edge for specified branch "${targetBranch.name}" on ${this.fullName}`}
 		}
 
 		// Ensure we're paused on a non-manual pause blockage
 		if (!edge.isBlocked) {
-			return { success: false, message: "Branch needs to have conflict to stomp changes" }
+			return { ...failureResult, message: "Branch needs to have conflict to stomp changes" }
 		}
 		
 		// Get the pause info of the edge
@@ -664,16 +671,16 @@ export class NodeBot extends PerforceStatefulBot implements NodeBotInterface {
 
 		// We've ensured we're on a blockage now. Make sure the CL is still valid.
 		if (blockageInfo.change !== changeCl) {
-			return { success: false, message: `Requested change CL "${changeCl}" does not match blockage change CL "${blockageInfo.change}"` }
+			return { ...failureResult, message: `Requested change CL "${changeCl}" does not match blockage change CL "${blockageInfo.change}"` }
 		}
 
 		let conflict = this.conflicts.getConflictByChangeCl(changeCl, targetBranch)
 
 		if (!conflict) {
-			return { success: false, message: `Unable to retrieve conflict for CL ${changeCl} merging to ${targetBranch.name}. Unable to stomp, please contact Robomerge support.` }
+			return { ...failureResult, message: `Unable to retrieve conflict for CL ${changeCl} merging to ${targetBranch.name}. Unable to stomp, please contact Robomerge support.` }
 		}
 		if (conflict.kind === 'Exclusive check-out') {
-			return { success: false, message: `Unable to stomp changes over exclusive checkout to files in ${targetBranch.name}. ` + 
+			return { ...failureResult, message: `Unable to stomp changes over exclusive checkout to files in ${targetBranch.name}. ` + 
 				`Unable to stomp, please revert the locked files in ${targetBranch.name} or contact Robomerge support.` }
 		}
 
@@ -682,7 +689,7 @@ export class NodeBot extends PerforceStatefulBot implements NodeBotInterface {
 
 		if (typeof blockageChangeRetrieval === 'string') {
 			return { 
-				success: false, 
+				...failureResult, 
 				message: `Unable to retrieve change infomation for given conflict. ${blockageChangeRetrieval}`
 			}
 		}
@@ -702,7 +709,7 @@ export class NodeBot extends PerforceStatefulBot implements NodeBotInterface {
 			const processResult = await this._createChangeInfo(blockageChange, edgeMap, this.p4.username, workspace, targetBranch)
 			if (!processResult.info) {
 				return {
-					success: false, 
+					...failureResult,  
 					message: `Could not parse blockage change\nErrors: ${processResult.errors}`
 				}
 			}
@@ -716,7 +723,7 @@ export class NodeBot extends PerforceStatefulBot implements NodeBotInterface {
 					message = `${message}\nErrors: ${mergeResult.errors}`
 				}
 				return { 
-					success: false, 
+					...failureResult,  
 					message
 				}
 			}
@@ -724,7 +731,7 @@ export class NodeBot extends PerforceStatefulBot implements NodeBotInterface {
 			const changelist = changes[0]!
 			if (changelist <= 0) {
 				return { 
-					success: false, 
+					...failureResult, 
 					message: `Unable to get integration shelf for CL ${changeCl}`
 				}
 			}
@@ -735,7 +742,7 @@ export class NodeBot extends PerforceStatefulBot implements NodeBotInterface {
 			}
 			catch (err) {
 				return { 
-					success: false, 
+					...failureResult, 
 					message: `Unable to get integration shelf for CL ${changeCl}: ${err}`
 				}
 			}
@@ -745,7 +752,7 @@ export class NodeBot extends PerforceStatefulBot implements NodeBotInterface {
 			// Unshelve 
 			if (!await this.p4.unshelve(workspace, integratedCl.change)) {
 				return { 
-					success: false, 
+					...failureResult,  
 					message: `Unable to unshelve CL ${integratedCl.change}`
 				}
 			}
@@ -755,13 +762,13 @@ export class NodeBot extends PerforceStatefulBot implements NodeBotInterface {
 
 			if (!resolveResult.hasConflict()) {
 				return { 
-					success: false, 
+					...failureResult,  
 					message: `No conflict found resolving CL ${changeCl}`
 				}
 			}
 			if (!resolveResult.successfullyParsedConflicts) {
 				return { 
-					success: false, 
+					...failureResult,  
 					message: `Encountered error enumerating conflicting files after integration attempt: ${resolveResult.getParsingError()}`
 				}
 			}
@@ -893,7 +900,8 @@ export class NodeBot extends PerforceStatefulBot implements NodeBotInterface {
 				message: "TBD",
 				nonBinaryFilesResolved,
 				remainingAllBinary,
-				svFiles
+				svFiles,
+				swarmURL: this.p4.swarmURL
 			}
 
 			if (!remainingAllBinary) {
@@ -1292,6 +1300,7 @@ export class NodeBot extends PerforceStatefulBot implements NodeBotInterface {
 	applyStatus(status: BranchStatus) {
 		status.display_name = this.displayName
 		status.last_cl = this.lastCl
+		status.swarmURL = this.p4.swarmURL
 
 		status.is_active = this.isActive
 		status.is_available = this.isAvailable
diff --git a/Engine/Extras/RoboMerge/v3/src/robo/status-types.ts b/Engine/Extras/RoboMerge/v3/src/robo/status-types.ts
--- a/Engine/Extras/RoboMerge/v3/src/robo/status-types.ts
+++ b/Engine/Extras/RoboMerge/v3/src/robo/status-types.ts
@@ -87,6 +87,7 @@ export type AnyStateInfo = ManualPauseInfo | BlockagePauseInfo | AvailableInfo
 export type BotStatusFields = Partial<PauseStatusFields> & {
 	display_name: string
 	last_cl: number
+	swarmURL?: string
 
 	headCL?: number
 
