Fix swarmURL pass back to client

diff --git a/Engine/Extras/RoboMerge/v3/src/robo/ipc.ts b/Engine/Extras/RoboMerge/v3/src/robo/ipc.ts
--- a/Engine/Extras/RoboMerge/v3/src/robo/ipc.ts
+++ b/Engine/Extras/RoboMerge/v3/src/robo/ipc.ts
@@ -539,8 +539,7 @@ export class IPC {
 						remainingAllBinary: stompVerification.remainingAllBinary,
 						files: stompVerification.svFiles,
 						validRequest: stompVerification.validRequest,
-						sourceSwarmURL: stompVerification.sourceSwarmURL,
-						targetSwarmURL: stompVerification.targetSwarmURL
+						swarmURL: stompVerification.swarmURL
 					} )
 				}
 			}
