p4 delete and revert commands are now resilient to the changelist not existing

diff --git a/Engine/Extras/RoboMerge/v3/src/common/perforce.ts b/Engine/Extras/RoboMerge/v3/src/common/perforce.ts
--- a/Engine/Extras/RoboMerge/v3/src/common/perforce.ts
+++ b/Engine/Extras/RoboMerge/v3/src/common/perforce.ts
@@ -1512,8 +1512,20 @@ export class PerforceContext {
 
 	// delete a CL
 	// output format is just error or not
-	deleteCl(roboWorkspace: RoboWorkspace, changelist: number, edgeServerAddress?: string) {
-		return this._execP4(roboWorkspace, ["change", "-d", changelist.toString()], {serverAddress: edgeServerAddress});
+	async deleteCl(roboWorkspace: RoboWorkspace, changelist: number, edgeServerAddress?: string) {
+		try {
+			await this._execP4(roboWorkspace, ["change", "-d", changelist.toString()], {serverAddress: edgeServerAddress})
+		}
+		catch (reason) {
+			if (!isExecP4Error(reason)) {
+				throw reason
+			}
+
+			let [_, output] = reason
+			if (!output.match(/Change (\d+) unknown\./)) {
+				throw reason
+			}
+		}
 	}
 
 	async getOpenedDetails(paths: {depotPath: string, name: string}[]) {
@@ -1620,7 +1632,8 @@ export class PerforceContext {
 
 			let [err, output] = reason
 			// this happens if there's literally nothing in the CL. consider this a success
-			if (!output.match(/file\(s\) not opened (?:on this client|in that changelist)\./)) {
+			if (!output.match(/file\(s\) not opened (?:on this client|in that changelist)\./) &&
+				!output.match(/Change (\d+) unknown\./)) {
 				throw err;
 			}
 		}
