Fix github commits not displaying

diff --git a/Engine/Extras/RoboMerge/v3/public/js/trackchange.js b/Engine/Extras/RoboMerge/v3/public/js/trackchange.js
--- a/Engine/Extras/RoboMerge/v3/public/js/trackchange.js
+++ b/Engine/Extras/RoboMerge/v3/public/js/trackchange.js
@@ -91,7 +91,8 @@ async function generateChangeList(dataObj) {
 		placedChanges.add(orderedChanges[i])
 	}
 	for (const cl of orderedChanges) {
-		const changeDetails = dataObj.changes[cl.toString()]
+		clKey = cl.toString()
+		const changeDetails = dataObj.changes[clKey]
 		html += '<tr valign="middle">'
 		html += `<td><b>${changeDetails.streamDisplayName}</b></td>`
 		if (changeDetails.swarmLink) {
@@ -99,7 +100,7 @@ async function generateChangeList(dataObj) {
 		} else {
 			html += `<td>CL#${cl}</td>`
 		}
-		html += `${formatGitHubCommit(gitHubCommits.get(cl))}`
+		html += `${formatGitHubCommit(gitHubCommits.get(clKey))}`
 		html += '</tr>'
 	}
 	html += "</tbody></table></div>"
