Fix fix

diff --git a/Engine/Extras/RoboMerge/v3/public/js/trackchange.js b/Engine/Extras/RoboMerge/v3/public/js/trackchange.js
--- a/Engine/Extras/RoboMerge/v3/public/js/trackchange.js
+++ b/Engine/Extras/RoboMerge/v3/public/js/trackchange.js
@@ -91,7 +91,7 @@ async function generateChangeList(dataObj) {
 		placedChanges.add(orderedChanges[i])
 	}
 	for (const cl of orderedChanges) {
-		clKey = cl.toString()
+		const clKey = cl.toString()
 		const changeDetails = dataObj.changes[clKey]
 		html += '<tr valign="middle">'
 		html += `<td><b>${changeDetails.streamDisplayName}</b></td>`
