Speed up TrackChanges using parallelization

diff --git a/Engine/Extras/RoboMerge/v3/src/robo/trackchange.ts b/Engine/Extras/RoboMerge/v3/src/robo/trackchange.ts
--- a/Engine/Extras/RoboMerge/v3/src/robo/trackchange.ts
+++ b/Engine/Extras/RoboMerge/v3/src/robo/trackchange.ts
@@ -152,13 +152,13 @@ export async function trackChange(robo: RoboMerge, logger: ContextualLogger, ini
 
     const isFTE = userTags.has("fte")
 
+    let changePromises: Promise<void>[] = []
     let clsToConsider: TrackedCL[] = [data.originalCL]
-    for (let i = 0; i < clsToConsider.length; i++) {
-        const clToConsider = clsToConsider[i]
 
+    const considerCL = async (clToConsider: TrackedCL) => {
         let changeToConsider = changes.get(JSON.stringify(clToConsider))
         if (!changeToConsider) {
-            continue
+            return
         }
 
         const externalServerTargets: Node[] = []
@@ -278,6 +278,7 @@ export async function trackChange(robo: RoboMerge, logger: ContextualLogger, ini
         for (const destCL of changeToConsider.destCLs) {
             if (!clsToConsider.find(cl => compareTrackedCL(cl, destCL))) {
                 clsToConsider.push(destCL)
+                changePromises.push(considerCL(destCL))
             }
         }
 
@@ -287,5 +288,10 @@ export async function trackChange(robo: RoboMerge, logger: ContextualLogger, ini
         }
     }
 
+    changePromises.push(considerCL(clsToConsider[0]))
+    for (const cp of changePromises) {
+        await cp
+    }
+
     return {success: true, data}
 }
\ No newline at end of file
