Speed up startup

diff --git a/Engine/Extras/RoboMerge/v3/src/robo/branchdefs.ts b/Engine/Extras/RoboMerge/v3/src/robo/branchdefs.ts
--- a/Engine/Extras/RoboMerge/v3/src/robo/branchdefs.ts
+++ b/Engine/Extras/RoboMerge/v3/src/robo/branchdefs.ts
@@ -369,10 +369,10 @@ export class BranchDefs {
 		branchGraph.branches = []
 
 		// Check for duplicate branch names
-		for (const def of branchesFromJSON) {
+		await Promise.all(branchesFromJSON.map(async (def) => {
 			if (!def.name) {
 				validationErrors.push(`Unable to parse branch definition: ${JSON.stringify(def)}`)
-				continue
+				return
 			}
 
 			branchGraph.branches.push(def)
@@ -380,7 +380,7 @@ export class BranchDefs {
 			const nameError = BranchDefs.checkName(def.name)
 			if (nameError) {
 				validationErrors.push(nameError)
-				continue
+				return
 			}
 
 			const upperName = def.name.toUpperCase()
@@ -405,7 +405,7 @@ export class BranchDefs {
 					validationErrors.push(`Stream ${streamResult.stream} not found`)
 				}
 			}
-		}
+		}))
 
 		// Check for duplicate aliases (and that branches/aliases are not in the ignore list)
 		const addAlias = (upperBranchName: string, upperAlias: string) => {
@@ -437,8 +437,7 @@ export class BranchDefs {
 			}
 		}
 
-		for (const def of branchGraph.branches) {
-
+		await Promise.all(branchGraph.branches.map(async (def) => {
 			const upperName = def.name!.toUpperCase()
 			if (def.aliases) {
 				for (const alias of def.aliases) {
@@ -453,18 +452,18 @@ export class BranchDefs {
 			// check all properties are known (could make things case insensitive here)
 			for (const keyName of Object.keys(def)) {
 				if (!nodeOptionFieldNames.has(keyName)) {
-					throw new Error(`Unknown property '${keyName}' specified for node ${def.name}`)
+					validationErrors.push(`Unknown property '${keyName}' specified for node ${def.name}`)
 				}
 			}
 
 			const results = await validateCommonOptions(logger, def, isPreviewing)
 			validationErrors.push(...results.errors)
 			validationWarnings.push(...results.warnings)
-		}
+		}))
 
 		// Check edge properties
 		if (branchGraph.edges) {
-			for (const edge of branchGraph.edges) {
+			await Promise.all(branchGraph.edges.map(async (edge) => {
 				if (!names.get(edge.from.toUpperCase())) {
 					validationErrors.push('Unrecognised source node in edge property ' + edge.from)
 				}
@@ -502,7 +501,7 @@ export class BranchDefs {
 				const results = await validateCommonOptions(logger, edge, isPreviewing)
 				validationErrors.push(...results.errors)
 				validationWarnings.push(...results.warnings)
-				}
+			}))
 		}
 
 		// Check flow
@@ -550,7 +549,7 @@ export class BranchDefs {
 
 		// Check branchspecs for valid branches
 		if (branchGraph.branchspecs) {
-			for (const spec of branchGraph.branchspecs) {
+			await Promise.all(branchGraph.branchspecs.map(async (spec) => {
 				let expectedFields = ['from', 'to', 'name']
 				for (const [key, val] of Object.entries(spec)) {
 					const efLength = expectedFields.length
@@ -583,7 +582,7 @@ export class BranchDefs {
 						validationErrors.push(`Branchspec ${spec.name} does not exist.`)
 					}
 				}
-			}
+			}))
 		}
 
 		if (validationErrors.length > 0) {
