Expose the serverID and make a generic comparison function that handles not specifying a serverID or explicitly specifying default

diff --git a/Engine/Extras/RoboMerge/v3/src/common/perforce.ts b/Engine/Extras/RoboMerge/v3/src/common/perforce.ts
--- a/Engine/Extras/RoboMerge/v3/src/common/perforce.ts
+++ b/Engine/Extras/RoboMerge/v3/src/common/perforce.ts
@@ -328,6 +328,7 @@ export interface ConflictedResolveNFile {
 
 interface ServerConfig {
 	username: string
+	serverID: string
 	serverAddress?: string
 	serverVersion: number
 	multiServerEnvironment: boolean
@@ -541,6 +542,7 @@ export class PerforceContext {
 	
 			let serverConfig: ServerConfig = {
 				username: resp.User,
+				serverID,
 				serverAddress,
 				serverVersion,
 				multiServerEnvironment: serversOutput.length > 1,
@@ -583,6 +585,10 @@ export class PerforceContext {
 		}
 	}
 
+	static sameServerID(aID?: string, bID?: string) {
+		return (aID || 'default') === (bID || 'default')
+	}
+
 	static getServerContext(logger: ContextualLogger, serverID?: string) {
 		serverID = serverID || 'default'
 		const server = this.servers.get(serverID)
@@ -600,6 +606,10 @@ export class PerforceContext {
 		return this._serverConfig.username
 	}
 
+	get serverID() {
+		return this._serverConfig.serverID
+	}
+
 	get serverAddress() {
 		return this._serverConfig.serverAddress
 	}
