Improve backtracking of source changes in cases where the #ROBOMERGE history has been culled

diff --git a/Engine/Extras/RoboMerge/v3/src/robo/trackchange.ts b/Engine/Extras/RoboMerge/v3/src/robo/trackchange.ts
--- a/Engine/Extras/RoboMerge/v3/src/robo/trackchange.ts
+++ b/Engine/Extras/RoboMerge/v3/src/robo/trackchange.ts
@@ -41,7 +41,7 @@ interface TrackChangeResults {
 }
 
 interface GatherOptions {
-    sourceCL?: TrackedCL,
+    parentCL?: TrackedCL,
     lastCL?: TrackedCL,
     initialSubmit?: boolean
 }
@@ -81,7 +81,7 @@ export async function trackChange(robo: RoboMerge, logger: ContextualLogger, ini
     }
 
     const getMergeMethod = (desc: string, cl: TrackedCL, opts?: GatherOptions): MergeMethod => {
-        if (opts?.sourceCL && cl.serverId != opts.sourceCL.serverId) {
+        if (opts?.parentCL && !PerforceContext.sameServerID(cl.serverId, opts.parentCL.serverId)) {
             return 'transfer'
         } else if (desc.includes("#ROBOMERGE-CONFLICT")) {
             return 'merge_with_conflict'
@@ -96,16 +96,66 @@ export async function trackChange(robo: RoboMerge, logger: ContextualLogger, ini
         }
     }
 
+    const getMoreDescEntries = async (p4: PerforceContext, clToConsider: TrackedCL, desc: DescribeResult) => {
+        // If the path is a root directory try and get a sampling 
+        // of files, otherwise just get an additional block to evaluate
+        if (desc.path.endsWith("/...")) {
+            const dirs = await p4.dirs(`${desc.path.slice(0,-3)}*@=${clToConsider.cl}`)
+            const filesPromises = dirs.map(dir => p4.files(`${dir}/...@=${clToConsider.cl}`, 1))
+            for (const fp of filesPromises) {
+                desc.entries.push(...(await fp))
+            }
+        }
+        else {
+            desc = await p4.describe(clToConsider.cl, 100)
+        }
+    }
+
     let changes = new Map<string, TrackedChange>()
 
+    let sourceChain: string[] = []
     let gatherCLInfo = async (cl: TrackedCL, opts?: GatherOptions) => {
         let p4 = PerforceContext.getServerContext(robo.roboMergeLogger, cl.serverId)
         try {
             let desc = await p4.describe(cl.cl, 1)
-            const mergeMethod = getMergeMethod(desc.description, cl, opts)
+            let initialSubmit = false
+            if (!opts?.parentCL) {
+                const sourceMatch = desc.description.match(/#ROBOMERGE-SOURCE: (.*)\n/g)
+                if (sourceMatch) {
+                    const newSources = Array.from(sourceMatch[0].matchAll(/CL (\d+)/g)).map(m => m[1]).reverse()
+                    sourceChain.push(...newSources)
+                    opts = {
+                        ...opts,
+                        parentCL: { cl: parseInt(newSources[0])/*, TODO: deal with cross server #ROBOMERGE-SOURCE tags. serverId: lastCL.serverId */ }
+                    }
+                }
+                else {
+                    for (const entry of desc.entries) {
+                        // integrated for move/delete will point at the paired move/add not where it was merged to in another stream, so not useful to evaluate it
+                        if (entry.action != 'move/delete') {
+                            // Do an integrated check
+                            const integrated = (await p4.integrated(null, entry.depotFile)).filter(x => x.how.includes('from') && x.toFile == entry.depotFile && x.endToRev == `#${entry.rev}`)
+                            if (integrated.length > 0) {
+                                const fstat = await p4.fstat(null, `${integrated[0].fromFile}${integrated[0].endFromRev}`)
+                                sourceChain.push(fstat[0].headChange)
+                                opts = {
+                                    ...opts,
+                                    parentCL: {cl: fstat[0].headChange, serverId: p4.serverID }
+                                }
+                            }
+                            break
+                        }
+                        else if (desc.entries.length == 1) {
+                            await getMoreDescEntries(p4, cl, desc)
+                        }
+                    }
+                }
+                initialSubmit = opts?.parentCL === undefined
+            }
+            const mergeMethod = getMergeMethod(desc.description, cl, {...opts, initialSubmit} )
             const clNode = getNode(desc)
             if (clNode || (RobomergeMethodStrings as readonly string[]).includes(mergeMethod)) {
-                changes.set(JSON.stringify(cl), {desc, mergeMethod, node: clNode, sourceCL: opts?.sourceCL, destCLs: (opts && opts.lastCL ? [opts.lastCL] : []) })
+                changes.set(JSON.stringify(cl), {desc, mergeMethod, node: clNode, sourceCL: opts?.parentCL, destCLs: (opts && opts.lastCL ? [opts.lastCL] : []) })
             }
         } 
         catch (reason) {
@@ -127,28 +177,24 @@ export async function trackChange(robo: RoboMerge, logger: ContextualLogger, ini
         return {success: false, message: `Unable to lookup CL# ${initialCL.cl}`}
     }
 
-    const sourceMatch = (changes.get(JSON.stringify(data.originalCL))?.desc as DescribeResult).description.match(/#ROBOMERGE-SOURCE: (.*)\n/g)
-    if (sourceMatch) {
-        let lastCL = initialCL
-        let originCL = initialCL
-        const matches = Array.from(sourceMatch[0].matchAll(/CL (\d+)/g))
-        for (let i=matches.length-1; i>=0;i--) {
-            const sourceCL = {cl: parseInt(matches[i][1])/*, TODO: deal with cross server #ROBOMERGE-SOURCE tags. serverId: lastCL.serverId */}
-            const parentCL = i > 0 ? {cl: parseInt(matches[i-1][1])} : undefined
-            if (await gatherCLInfo(sourceCL, {sourceCL: parentCL, lastCL, initialSubmit: i==0})) {
-                if (i == matches.length-1) {
-                    let initialChange = changes.get(JSON.stringify(initialCL))!
-                    // Fix up the initial changes merge method and source now that we know it is not the initial submit
-                    initialChange.mergeMethod = getMergeMethod(initialChange.desc.description, initialCL, {sourceCL})
-                    initialChange.sourceCL = sourceCL
-                }
-                originCL = sourceCL
+    let lastCL = initialCL
+    let originCL = initialCL
+    for (let i=0; i < sourceChain.length; i++) {
+        const sourceCL = {cl: parseInt(sourceChain[i])/*, TODO: deal with cross server #ROBOMERGE-SOURCE tags. serverId: lastCL.serverId */}
+        const parentCL = (sourceChain[i+1] && {cl: parseInt(sourceChain[i+1])}) as TrackedCL
+        if (await gatherCLInfo(sourceCL, {parentCL, lastCL})) {
+            if (i == 0) {
+                let initialChange = changes.get(JSON.stringify(initialCL))!
+                // Fix up the initial changes merge method and source now that we know it is not the initial submit
+                initialChange.mergeMethod = getMergeMethod(initialChange.desc.description, initialCL, {parentCL: sourceCL})
+                initialChange.sourceCL = sourceCL
             }
-            lastCL = sourceCL
+            originCL = sourceCL
         }
-        // We've now determined the earliest CL we are going to start tracking from
-        data.originalCL = originCL
+        lastCL = sourceCL
     }
+    // We've now determined the earliest CL we are going to start tracking from
+    data.originalCL = originCL
 
     const isFTE = userTags.has("fte")
 
@@ -225,7 +271,7 @@ export async function trackChange(robo: RoboMerge, logger: ContextualLogger, ini
                         if (change.desc.includes(clToConsider.cl.toString())) {
                             const newChangeToConsider = {cl: change.change, serverId: externalTarget.serverID}
                             changeToConsider.destCLs.push(newChangeToConsider)
-                            await gatherCLInfo(newChangeToConsider, {sourceCL: clToConsider})
+                            await gatherCLInfo(newChangeToConsider, {parentCL: clToConsider})
                             externalServerTargets.splice(i, 1)
                             break
                         }
@@ -247,7 +293,7 @@ export async function trackChange(robo: RoboMerge, logger: ContextualLogger, ini
                                 } else {
                                     const newChangeToConsider = {cl: integ.change, serverId: clToConsider.serverId}
                                     changeToConsider.destCLs.push(newChangeToConsider)
-                                    await gatherCLInfo(newChangeToConsider, {sourceCL: clToConsider})
+                                    await gatherCLInfo(newChangeToConsider, {parentCL: clToConsider})
                                 }
                             }
                         }
@@ -262,17 +308,8 @@ export async function trackChange(robo: RoboMerge, logger: ContextualLogger, ini
                         (RobomergeMethodStrings as readonly string[]).includes(changeToConsider.mergeMethod)) {
                 // If we only have 1 entry and we didn't get integration info off of it
                 // and the graph suggests we are expecting there could be other changes
-                // get more of the files. If the path is a root directory try and get a sampling 
-                // of files, otherwise just get an additional block to evaluate
-                if (changeToConsider.desc.path.endsWith("/...")) {
-                    const dirs = await p4.dirs(`${changeToConsider.desc.path.slice(0,-3)}*@=${clToConsider.cl}`)
-                    for (const dir of dirs) {
-                        changeToConsider.desc.entries.push(...(await p4.files(`${dir}/...@=${clToConsider.cl}`, 1)))
-                    }
-                }
-                else {
-                    changeToConsider.desc = await p4.describe(clToConsider.cl, 100)
-                }
+                // get more of the files. 
+                await getMoreDescEntries(p4, clToConsider, changeToConsider.desc)
             }
         }
         for (const destCL of changeToConsider.destCLs) {
