Gracefully handle p4.files saying there are no files and just return an empty list

diff --git a/Engine/Extras/RoboMerge/v3/src/common/perforce.ts b/Engine/Extras/RoboMerge/v3/src/common/perforce.ts
--- a/Engine/Extras/RoboMerge/v3/src/common/perforce.ts
+++ b/Engine/Extras/RoboMerge/v3/src/common/perforce.ts
@@ -1820,11 +1820,25 @@ export class PerforceContext {
 		return (await this.execAndParse(null, ['dirs', path])).map((dir: any) => dir.dir)
 	}
 
-	files(path: string, maxFiles?: number) {
+	async files(path: string, maxFiles?: number) {
 		const args = ['files',
 			...(maxFiles ? ['-m', maxFiles.toString()] : []),
 			path]
-		return this.execAndParse(null, args);
+		try {
+			return await this.execAndParse(null, args);
+		}
+		catch (reason) {
+			if (!isExecP4Error(reason)) {
+				throw reason
+			}
+
+			let [_, output] = reason
+			// If perforce doesn't detect revisions in the given range, return an empty set of revisions
+			if (output.trim().endsWith("no such file(s).")) {
+				return []
+			}
+			throw reason		
+		}
 	}
 
 	sizes(path: string, summary?: boolean) {
