Fix tracking in specific cases where there are a number of move/delete records first in the describe results
Track out of streams that are no longer in the graph for FTEs

diff --git a/Engine/Extras/RoboMerge/v3/src/robo/trackchange.ts b/Engine/Extras/RoboMerge/v3/src/robo/trackchange.ts
--- a/Engine/Extras/RoboMerge/v3/src/robo/trackchange.ts
+++ b/Engine/Extras/RoboMerge/v3/src/robo/trackchange.ts
@@ -100,10 +100,17 @@ export async function trackChange(robo: RoboMerge, logger: ContextualLogger, ini
         // If the path is a root directory try and get a sampling 
         // of files, otherwise just get an additional block to evaluate
         if (desc.path.endsWith("/...")) {
-            const dirs = await p4.dirs(`${desc.path.slice(0,-3)}*@=${clToConsider.cl}`)
-            const filesPromises = dirs.map(dir => p4.files(`${dir}/...@=${clToConsider.cl}`, 1))
-            for (const fp of filesPromises) {
-                desc.entries.push(...(await fp))
+            const rootDir = `${desc.path.slice(0,-3)}*@=${clToConsider.cl}`
+            let fpIndex = 0
+            let filesPromises = [p4.files(rootDir, 5)]
+            while (desc.entries.length <= 1) {
+                const dirs = await p4.dirs(rootDir)
+                filesPromises.push(...dirs.slice(0, 5).map(dir => p4.files(`${dir}/...@=${clToConsider.cl}`, 5)))
+                while (fpIndex < filesPromises.length) {
+                    desc.entries.push(...(await filesPromises[fpIndex]))
+                    fpIndex++
+                }
+                desc.entries = desc.entries.filter(e => e.action != 'move/delete')
             }
         }
         else {
@@ -235,12 +242,15 @@ export async function trackChange(robo: RoboMerge, logger: ContextualLogger, ini
                     externalServerTargets.push(edge.target)
                 }
             }
-        } else if (isFTE && botFilter.length == 0) {
-            includeInResults = true
-            if (changeToConsider.desc.path) {
-                streamDisplayName = getStreamFromPath(changeToConsider.desc.path) || changeToConsider.desc.path
-            } else if (changeToConsider.desc.entries.length > 0) {
-                streamDisplayName = getStreamFromPath(changeToConsider.desc.entries[0].depotFile) || streamDisplayName
+        } else if (isFTE) {
+            hasAutomergeTarget = true // we can't know this from the graph, but it doesn't hurt (much) to look
+            if (botFilter.length == 0) {
+                includeInResults = true
+                if (changeToConsider.desc.path) {
+                    streamDisplayName = getStreamFromPath(changeToConsider.desc.path) || changeToConsider.desc.path
+                } else if (changeToConsider.desc.entries.length > 0) {
+                    streamDisplayName = getStreamFromPath(changeToConsider.desc.entries[0].depotFile) || streamDisplayName
+                }
             }
         }
 
