Allow track changes to initiate a search with a CL from a remote server

diff --git a/Engine/Extras/RoboMerge/v3/public/js/boilerplate.js b/Engine/Extras/RoboMerge/v3/public/js/boilerplate.js
--- a/Engine/Extras/RoboMerge/v3/public/js/boilerplate.js
+++ b/Engine/Extras/RoboMerge/v3/public/js/boilerplate.js
@@ -539,11 +539,17 @@ function generateRobomergeFooter() {
 		})
 		if (data) {
 			// Ensure they entered a CL
-			if (isNaN(parseInt(data.cl))) {
-				displayErrorMessage("Please provide a valid changelist number to track.")
-				return
+			if (!isNaN(parseInt(data.cl))) {
+				window.open(`/trackchange/${data.cl}`, '_blank') 
+			}
+			else {
+				const splitCL = data.cl.split(" ")
+				if (splitCL.length != 2 || isNaN(splitCL[1])) {
+					displayErrorMessage("Please provide a valid changelist number to track.")
+					return
+				}
+				window.open(`/trackchange/${splitCL[0]}/${splitCL[1]}`)
 			}
-			window.open(`/trackchange/${data.cl}`, '_blank') 
 		}
 	})
 	trackChangeButton.text("Track Change")
diff --git a/Engine/Extras/RoboMerge/v3/src/robo/ipc.ts b/Engine/Extras/RoboMerge/v3/src/robo/ipc.ts
--- a/Engine/Extras/RoboMerge/v3/src/robo/ipc.ts
+++ b/Engine/Extras/RoboMerge/v3/src/robo/ipc.ts
@@ -174,9 +174,13 @@ export class IPC {
 			return {statusCode: 400, message: 'No CL parameter provided.'}
 		}
 
-		const cl = parseInt(clStr)
-		if (isNaN(cl)) {
-			return {statusCode: 400, message: `Invalid CL parameter: ${cl}`}
+		const trackedCL = {cl: parseInt(clStr), serverId: queryObj.serverId}
+		if (isNaN(trackedCL.cl)) {
+			return {statusCode: 400, message: `Invalid CL parameter: ${clStr}`}
+		}
+
+		if (trackedCL.serverId && !PerforceContext.getServerConfig(trackedCL.serverId)) {
+			return {statusCode: 400, message: `Invalid server parameter: ${trackedCL.serverId}`}
 		}
 
 		if (!this.robo.graph) {
@@ -188,7 +192,7 @@ export class IPC {
 			userTags.add(tag)
 		}
 	
-		const result  = await trackChange(this.robo, this.ipcLogger, {cl}, userTags, queryObj)
+		const result  = await trackChange(this.robo, this.ipcLogger, trackedCL, userTags, queryObj)
 
 		if (result.success) {
 			return {...OPERATION_SUCCESS, data: result.data }
diff --git a/Engine/Extras/RoboMerge/v3/src/robo/roboserver.ts b/Engine/Extras/RoboMerge/v3/src/robo/roboserver.ts
--- a/Engine/Extras/RoboMerge/v3/src/robo/roboserver.ts
+++ b/Engine/Extras/RoboMerge/v3/src/robo/roboserver.ts
@@ -216,6 +216,29 @@ class RoboWebApp implements AppInterface {
 		return Mustache.render(template, {cl, query})
 	}
 
+	@SecureHandler('GET', '/trackchange/*/*', {filetype: 'text/html'}) 
+	async trackChangeOnServer(serverId: string, clStr: string) {
+		if (!this.authData) {
+			throw new Error('Secure call but no auth data?')
+		}
+
+		const cl = parseInt(clStr)
+		if (isNaN(cl)) {
+			throw new Error(`Failed to parse alleged CL '${clStr}'`)
+		}
+
+		let query = "/trackchange"
+		if (this.request.url.search.length > 0) {
+			query += `${this.request.url.search}&`
+		} else {
+			query += "?"
+		}
+		query += `cl=${clStr}&serverId=${serverId}`
+
+		const template = await readUtf8File('public/trackchange.html')
+		return Mustache.render(template, {cl, query})
+	}
+
 	@SecureHandler('GET', '/trackchange') 
 	getTrackedChanges() {
 		if (!this.authData) {
