Make it clearer instead of an error when there are unexpected mergeMethods
Output source when there is a parse issue in the viz graph source

diff --git a/Engine/Extras/RoboMerge/v3/public/js/flow.js b/Engine/Extras/RoboMerge/v3/public/js/flow.js
--- a/Engine/Extras/RoboMerge/v3/public/js/flow.js
+++ b/Engine/Extras/RoboMerge/v3/public/js/flow.js
@@ -18,7 +18,13 @@ function setDefault(map, key, def) {
 
 function renderGraph(src) {
     return Viz.instance().then(function (viz) {
-        return viz.renderSVGElement(src);
+        try {
+            return viz.renderSVGElement(src);
+        }
+        catch (err) {
+            console.log(src)
+            throw err
+        }
     })
 }
 
diff --git a/Engine/Extras/RoboMerge/v3/public/js/trackchange.js b/Engine/Extras/RoboMerge/v3/public/js/trackchange.js
--- a/Engine/Extras/RoboMerge/v3/public/js/trackchange.js
+++ b/Engine/Extras/RoboMerge/v3/public/js/trackchange.js
@@ -6,7 +6,8 @@ function getMergeMethodStyle(mergeMethod) {
 		case "automerge":
 			return {color: "black", style: "solid"}
 		case "initialSubmit":
-			return {color: "", style: ""}
+			console.log("Unexpected mergeMethod of initialSubmit")
+			return {color: "pink", style: "solid"}
 		case "merge_with_conflict":
 			return {color: "red", style: "dashed"}
 		case "manual_merge":
@@ -16,7 +17,8 @@ function getMergeMethodStyle(mergeMethod) {
 		case "transfer":
 			return {color: "orange", style: "solid"}
 	}
-	return "UNKNOWN CASE: " + mergeMethod	
+	console.log("UNKNOWN CASE: " + mergeMethod)
+	return {color: "pink", style: "solid"}
 }
 
 function getStreamGraphName(streamDisplayName) {
