Fix source chain tracking when the initiating change is on a remote server

diff --git a/Engine/Extras/RoboMerge/v3/src/robo/trackchange.ts b/Engine/Extras/RoboMerge/v3/src/robo/trackchange.ts
--- a/Engine/Extras/RoboMerge/v3/src/robo/trackchange.ts
+++ b/Engine/Extras/RoboMerge/v3/src/robo/trackchange.ts
@@ -87,7 +87,7 @@ export async function trackChange(robo: RoboMerge, logger: ContextualLogger, ini
             return 'merge_with_conflict'
         } else if (desc.includes("#ROBOMERGE-SOURCE")) {
             return 'automerge'
-        } else if (opts?.initialSubmit || compareTrackedCL(cl, data.originalCL)) {
+        } else if (opts?.initialSubmit) {
             return 'initialSubmit'
         } else if (desc.includes("Populate")) {
             return 'populate'
@@ -120,7 +120,7 @@ export async function trackChange(robo: RoboMerge, logger: ContextualLogger, ini
 
     let changes = new Map<string, TrackedChange>()
 
-    let sourceChain: string[] = []
+    let sourceChain: TrackedCL[] = []
     let gatherCLInfo = async (cl: TrackedCL, opts?: GatherOptions) => {
         let p4 = PerforceContext.getServerContext(robo.roboMergeLogger, cl.serverId)
         try {
@@ -129,12 +129,9 @@ export async function trackChange(robo: RoboMerge, logger: ContextualLogger, ini
             if (!opts?.parentCL) {
                 const sourceMatch = desc.description.match(/#ROBOMERGE-SOURCE: (.*)\n/g)
                 if (sourceMatch) {
-                    const newSources = Array.from(sourceMatch[0].matchAll(/CL (\d+)/g)).map(m => m[1]).reverse()
+                    const newSources: TrackedCL[] = Array.from(sourceMatch[0].matchAll(/CL (\d+)/g)).map(m => {return {cl: parseInt(m[1])/*, TODO: deal with cross server #ROBOMERGE-SOURCE tags. serverId: lastCL.serverId */ }}).reverse()
                     sourceChain.push(...newSources)
-                    opts = {
-                        ...opts,
-                        parentCL: { cl: parseInt(newSources[0])/*, TODO: deal with cross server #ROBOMERGE-SOURCE tags. serverId: lastCL.serverId */ }
-                    }
+                    opts = { ...opts, parentCL: newSources[0] }
                 }
                 else {
                     for (const entry of desc.entries) {
@@ -144,11 +141,8 @@ export async function trackChange(robo: RoboMerge, logger: ContextualLogger, ini
                             const integrated = (await p4.integrated(null, entry.depotFile)).filter(x => x.how.includes('from') && x.toFile == entry.depotFile && x.endToRev == `#${entry.rev}`)
                             if (integrated.length > 0) {
                                 const fstat = await p4.fstat(null, `${integrated[0].fromFile}${integrated[0].endFromRev}`)
-                                sourceChain.push(fstat[0].headChange)
-                                opts = {
-                                    ...opts,
-                                    parentCL: {cl: fstat[0].headChange, serverId: p4.serverID }
-                                }
+                                opts = { ...opts, parentCL: {cl: fstat[0].headChange, serverId: p4.serverID } }
+                                sourceChain.push(opts.parentCL!)
                             }
                             break
                         }
@@ -180,15 +174,15 @@ export async function trackChange(robo: RoboMerge, logger: ContextualLogger, ini
         return true
     }
 
-    if (!(await gatherCLInfo(initialCL))) {
+    if (!(await gatherCLInfo(initialCL, {initialSubmit: true}))) {
         return {success: false, message: `Unable to lookup CL# ${initialCL.cl}`}
     }
 
     let lastCL = initialCL
     let originCL = initialCL
     for (let i=0; i < sourceChain.length; i++) {
-        const sourceCL = {cl: parseInt(sourceChain[i])/*, TODO: deal with cross server #ROBOMERGE-SOURCE tags. serverId: lastCL.serverId */}
-        const parentCL = (sourceChain[i+1] && {cl: parseInt(sourceChain[i+1])}) as TrackedCL
+        const sourceCL = sourceChain[i]
+        const parentCL = sourceChain[i+1]
         if (await gatherCLInfo(sourceCL, {parentCL, lastCL})) {
             if (i == 0) {
                 let initialChange = changes.get(JSON.stringify(initialCL))!
