Backtrack across servers

diff --git a/Engine/Extras/RoboMerge/v3/src/robo/trackchange.ts b/Engine/Extras/RoboMerge/v3/src/robo/trackchange.ts
--- a/Engine/Extras/RoboMerge/v3/src/robo/trackchange.ts
+++ b/Engine/Extras/RoboMerge/v3/src/robo/trackchange.ts
@@ -122,17 +122,37 @@ export async function trackChange(robo: RoboMerge, logger: ContextualLogger, ini
 
     let sourceChain: TrackedCL[] = []
     let gatherCLInfo = async (cl: TrackedCL, opts?: GatherOptions) => {
+
+        if (cl.serverId === undefined && opts?.lastCL) {
+            const lastChange = changes.get(JSON.stringify(opts.lastCL))
+            if (lastChange?.node) {
+                let edges = graph.getEdgesForNode(lastChange.node)
+                const possibleServers = new Set([...edges].filter(e => e.target.stream === lastChange.node!.stream).map(e => e.source.serverID))
+                if (possibleServers.size == 1) {
+                    cl.serverId = [...possibleServers][0]
+                } 
+                else {
+                    // TODO: figure out which one it is
+                }
+            }    
+        }
+
         let p4 = PerforceContext.getServerContext(robo.roboMergeLogger, cl.serverId)
         try {
             let desc = await p4.describe(cl.cl, 1)
             let initialSubmit = false
             if (!opts?.parentCL) {
-                const sourceMatch = desc.description.match(/#ROBOMERGE-SOURCE: (.*)\n/g)
-                if (sourceMatch) {
-                    const newSources: TrackedCL[] = Array.from(sourceMatch[0].matchAll(/CL (\d+)/g)).map(m => {return {cl: parseInt(m[1])/*, TODO: deal with cross server #ROBOMERGE-SOURCE tags. serverId: lastCL.serverId */ }}).reverse()
+                //TODO: abstract this in to a description parser per stream/branch/server
+                const sourceMatch = desc.description.match(/#ROBOMERGE-SOURCE: (?<RMSource>.*)\n|Merge-CL: (?<MergeCL>\d*)\n/)
+                if (sourceMatch?.groups!.RMSource) {
+                    const newSources: TrackedCL[] = Array.from(sourceMatch.groups.RMSource.matchAll(/CL (\d+)/g)).map(m => {return {cl: parseInt(m[1])}}).reverse()
                     sourceChain.push(...newSources)
                     opts = { ...opts, parentCL: newSources[0] }
                 }
+                else if (sourceMatch?.groups!.MergeCL) {
+                    opts = { ...opts, parentCL: {cl:parseInt(sourceMatch.groups.MergeCL)}}
+                    sourceChain.push(opts.parentCL!)
+                }
                 else {
                     for (const entry of desc.entries) {
                         // integrated for move/delete will point at the paired move/add not where it was merged to in another stream, so not useful to evaluate it
