Safely compare serverIds

diff --git a/Engine/Extras/RoboMerge/v3/src/robo/trackchange.ts b/Engine/Extras/RoboMerge/v3/src/robo/trackchange.ts
--- a/Engine/Extras/RoboMerge/v3/src/robo/trackchange.ts
+++ b/Engine/Extras/RoboMerge/v3/src/robo/trackchange.ts
@@ -23,7 +23,7 @@ interface TrackedCL {
 }
 
 function compareTrackedCL(a: TrackedCL, b: TrackedCL) {
-    return a.cl == b.cl && a.serverId == b.serverId
+    return a.cl == b.cl && PerforceContext.sameServerID(a.serverId, b.serverId)
 }
 
 interface TrackedChange {
@@ -252,7 +252,7 @@ export async function trackChange(robo: RoboMerge, logger: ContextualLogger, ini
                         hasAutomergeTarget = edge.source.stream == changeToConsider.node.stream
                     }
                 }
-                if (clToConsider.serverId != edge.target.serverID) {
+                if (!PerforceContext.sameServerID(clToConsider.serverId, edge.target.serverID)) {
                     externalServerTargets.push(edge.target)
                 }
             }
