Catch any track changes errors and report it as an error

diff --git a/Engine/Extras/RoboMerge/v3/src/robo/ipc.ts b/Engine/Extras/RoboMerge/v3/src/robo/ipc.ts
--- a/Engine/Extras/RoboMerge/v3/src/robo/ipc.ts
+++ b/Engine/Extras/RoboMerge/v3/src/robo/ipc.ts
@@ -192,7 +192,16 @@ export class IPC {
 			userTags.add(tag)
 		}
 	
-		const result  = await trackChange(this.robo, this.ipcLogger, trackedCL, userTags, queryObj)
+		var result
+		try {
+			result  = await trackChange(this.robo, this.ipcLogger, trackedCL, userTags, queryObj)
+		}
+		catch (reason) {
+			result = {
+				success: false,
+				message: reason
+			}
+		}
 
 		if (result.success) {
 			return {...OPERATION_SUCCESS, data: result.data }
