Handle integrated case where the bot does not have visibility on the integrated paths and gets a no permission for operation on file(s). response

diff --git a/Engine/Extras/RoboMerge/v3/src/common/perforce.ts b/Engine/Extras/RoboMerge/v3/src/common/perforce.ts
--- a/Engine/Extras/RoboMerge/v3/src/common/perforce.ts
+++ b/Engine/Extras/RoboMerge/v3/src/common/perforce.ts
@@ -1971,7 +1971,7 @@ export class PerforceContext {
 
 			let [err, output] = reason
 			// If perforce doesn't detect revisions in the given range, return an empty set of revisions
-			if (output.includes("no file(s) integrated.")) {
+			if (output.includes("no file(s) integrated.") || output.includes("no permission for operation on file(s).")) {
 				return []
 			}
 			throw err
