Strip Settings of unused functionality and old migration code

diff --git a/Engine/Extras/RoboMerge/v3/src/robo/graphbot.ts b/Engine/Extras/RoboMerge/v3/src/robo/graphbot.ts
--- a/Engine/Extras/RoboMerge/v3/src/robo/graphbot.ts
+++ b/Engine/Extras/RoboMerge/v3/src/robo/graphbot.ts
@@ -92,7 +92,7 @@ export class GraphBot implements BotEventHandler {
 			graphBot.botLogger.error(`Problem starting up bot ${botname}: ${error}`);
 		}
 
-		graphBot._settings = new Settings(botname, graphBot.branchGraph, graphBot.botLogger, graphBot.p4)
+		graphBot._settings = new Settings(botname, graphBot.botLogger, graphBot.p4)
 
 		graphBot.externalUrl = externalUrl
 
diff --git a/Engine/Extras/RoboMerge/v3/src/robo/settings.ts b/Engine/Extras/RoboMerge/v3/src/robo/settings.ts
--- a/Engine/Extras/RoboMerge/v3/src/robo/settings.ts
+++ b/Engine/Extras/RoboMerge/v3/src/robo/settings.ts
@@ -4,9 +4,6 @@ import * as fs from 'fs';
 import { Args } from '../common/args';
 import { ContextualLogger } from '../common/logger';
 import { VersionReader } from '../common/version';
-import { BranchGraph } from './branchgraph';
-import { PersistentConflict } from './conflict-interfaces';
-import { NOTIFICATIONS_PERSISTENCE_KEY, SLACK_MESSAGES_PERSISTENCE_KEY } from './notifications';
 import { PerforceContext, Workspace } from '../common/perforce';
 import semver = require('semver')
 
@@ -17,6 +14,8 @@ import semver = require('semver')
  * 2.2.2 - No change to settings structure -- Going to match Robomerge versions from now on -- May 2019
  * 3.0.0 - Branchbot.ts is split into NodeBot and EdgeBot classes, adding edges settings.
  *       - Reworked pause object into seperate 'manual pause' and 'blockage' objects, need to migrate old data - Sept. '19
+ * 3.0.2 - Fix for pause state serialisation.
+ * 3.1.0 - Remove backwards compatibility before 3.0.2
  */
 const VERSION = VersionReader.getPackageVersion()
 
@@ -37,12 +36,6 @@ function readFileToString(filename: string) {
 	}
 }
 
-interface FieldToSet {
-	name: string
-	value: any
-	overrideContext?: string
-}
-
 export class Context {
 	readonly path: string[];
 	protected object: { [key: string]: any };
@@ -76,17 +69,6 @@ export class Context {
 		let val = parseInt(this.get(name));
 		return isNaN(val) ? (dflt === undefined ? 0 : dflt) : val;
 	}
-	getFloat(name: string, dflt?: number) {
-		let val = parseFloat(this.get(name));
-		return isNaN(val) ? (dflt === undefined ? 0 : dflt) : val;
-	}
-	getString(name: string, dflt?: string) {
-		let val = this.get(name);
-		if (val === undefined || val === null) {
-			return (dflt === undefined ? null : dflt);
-		}
-		return val.toString();
-	}
 	get(name: string) {
 		return this.object[name];
 	}
@@ -98,14 +80,6 @@ export class Context {
 		this.object[name] = value;
 		this.settings._saveObject();
 	}
-	setMultiple(fields: FieldToSet[]) {
-		for (const field of fields) {
-			const context = field.overrideContext ? this.settings.getContext(field.overrideContext) : this
-			context.set(field.name, field.value)
-		}
-
-		this.settings._saveObject()
-	}
 }
 
 type FilePurpose = 
@@ -131,7 +105,7 @@ export class Settings {
 	private readonly settingsLogger: ContextualLogger
 	private readonly p4workspace?: Workspace
 
-	constructor(botname: string, branchgraph: BranchGraph, parentLogger: ContextualLogger, private p4: PerforceContext) {
+	constructor(botname: string, parentLogger: ContextualLogger, private p4: PerforceContext) {
 		this.settingsLogger = parentLogger.createChild('Settings')
 
 		this.lastPersistBackupTime = new Date(Date.now())
@@ -176,42 +150,16 @@ export class Settings {
 			}
 		}
 
-		// if we can't write, don't bother with migration code.
-		if (!this.enableSave) {
-			return
+		if (semver.lt(objVersion, '3.0.2')) {
+			throw new Error(`Settings files prior to 3.0.2 are unsupported`)
 		}
 
-
 		/**
 		 * MIGRATION CODE SECTION START
 		 */
 		// For comparison later
 		const previousObjVersion = new semver.SemVer(objVersion)
 
-		// Check if we need to run pre-1.0 data migrations
-		if (semver.lt(objVersion, '1.0.0')) {
-			this.migrateSlackMessageDataKeys(branchgraph)
-			objVersion = new semver.SemVer('1.0.0')
-		}
-
-		// Don't use 2.0 -- post 1.0, version of settings will match Robomerge package version
-
-		if (semver.lt(objVersion, '3.0.0')) {
-			// Convert the old node centralized pause info to edge-based statuses
-			this.migratePauseInfoResetLastCl()
-
-			// Our URLs have changes format, we need to ensure our Slack messages are refreshed
-			this.removeAllSlackMessages()
-
-			objVersion = new semver.SemVer('3.0.0')
-		}
-
-		if (semver.lt(objVersion, '3.0.2')) {
-			this.migrateToFixBrokenPauseInfos()
-
-			objVersion = new semver.SemVer('3.0.2')
-		}
-
 		// We're up to date!
 		if (previousObjVersion.compare(VERSION) !== 0) {
 			this.settingsLogger.info(`Changing settings version from '${previousObjVersion.raw}' to '${VERSION.raw}'`)
@@ -225,202 +173,6 @@ export class Settings {
 		this._saveObject(true);
 	}
 
-	/**
-	 * Pre-1.0, Slack Messages were indexed by "Changelist:Branch".
-	 * With 1.0, they are now indexed with "Changelist:Branch:Channel Name"
-	 */
-	private migrateSlackMessageDataKeys(branchgraph: BranchGraph) {
-		this.settingsLogger.info("Migrating old Slack Message keys in settings data.")
-		// Check for slack messages
-		if (this.object.notifications && this.object.notifications.slackMessages) {
-			let slackMessages : {[key: string]: any} = this.object.notifications.slackMessages
-			// Regex to find a CL # paired with a branch
-			const oldKeyRegex = /^(\d+):([^:]+)$/
-			const newKeyRegex = /^(\d+):([^:]+):([^:]+)$/
-
-			// Iterate over slack message key names and change their format
-			for (let key in slackMessages) {
-				let matches = oldKeyRegex.exec(key)
-				if (!matches) {
-					// We didn't match the expected old style.
-					// Maybe we match the new style, and this configuration was partially migrated
-					let newMatches = newKeyRegex.exec(key)
-					if (newMatches) {
-						// Ah! We matched the format of the new key. Let's skip this.
-						this.settingsLogger.verbose(`Skipping key "${key}"`)
-						continue
-					}
-
-					// Didn't match either format. What gives? Get a dev to triage it.
-					throw new Error(`Error migrating slack message data for message "${key}"`)
-				}
-
-				// Re-add the contents of this message to a new value
-				const newKey = `${matches[1]}:${matches[2]}:${branchgraph.config.slackChannel}`
-				this.settingsLogger.info(`Migrating key "${key}" to "${newKey}"`)
-				this.object.notifications.slackMessages[newKey] = {
-					...slackMessages[key],
-					"channel": branchgraph.config.slackChannel
-				}
-
-				// Remove the old one
-				delete slackMessages[key]
-			}
-		}
-		this.settingsLogger.info("Slack Message key migration complete.")
-	}
-
-	/**
-	 * With 3.0, we seperate the old "PauseInfo" interface into "ManualPauseInfo" and "BlockagePauseInfo" so that
-	 * we can pause blocked bots. Before, a blocked bot could not be paused because PauseInfo could only hold one state.
-	 * Also there seems to be a small bug with Edgebot lastCl < 3.0, so we'll reset to the node's CL
-	 */
-	private migratePauseInfoResetLastCl() {
-		this.settingsLogger.info("Migrating old pause info in settings data.")
-
-		// Find each Nodebot and change their pause info block to the new format
-		for (const key in this.object) {
-			// These are keys in the settings object that are not nodes. Ignore.
-			if (key.toUpperCase() === "VERSION" || key.toUpperCase() === "NOTIFICATIONS") {
-				continue
-			}
-
-			const node = this.object[key];
-			this.migratePauseInfo_Helper(key, node)
-
-			const nodeLastCl = node.lastCl
-			const edges = node.edges
-			if (!edges) {
-				continue
-			}
-
-			// Previous EdgeBot lastCl values are incorrect from 2.x.
-			// Reset them here to NodeBot lastCl for the migration 2.x -> 3.x
-			for (const edgeName in edges) {
-				this.settingsLogger.info(`\tSetting ${key}->${edgeName} lastCl=${nodeLastCl}`)
-				edges[edgeName].lastCl = nodeLastCl
-			}
-		}
-
-		this.settingsLogger.info("Pause Info migration complete.")
-	}
-
-	/**
-	 * In 3.0.1 we were serializing PauseStates directly, rather than running the toJSON function
-	 */
-	private migrateToFixBrokenPauseInfos() {
-		for (const key in this.object) {
-			// These are keys in the settings object that are not nodes. Ignore.
-			if (key.toUpperCase() === "VERSION" || key.toUpperCase() === "NOTIFICATIONS") {
-				continue
-			}
-
-			this.migratePauseState_Helper_runToJSON(key, this.object[key])
-		}
-	}
-
-
-	// Go through node data and migrate to new format. Return any pause data that needs to be applied to the edge
-	private migratePauseInfo_Helper(nodeName: string, nodeData: { [key: string]: any } ): any {
-		// Check for the existence of the old "info" object
-		if (nodeData.pause && nodeData.pause.info) {
-			this.settingsLogger.info(`Migrating ${nodeName} pause info...`)
-
-			let oldPauseInfo = nodeData.pause.info
-
-			// Handle manual pauses
-			if (oldPauseInfo.type === 'manual-lock') {
-				nodeData.pause.manual_pause = oldPauseInfo
-			}
-			// Handle conflicts
-			else if (oldPauseInfo.type === 'branch-stop') {
-				// We should have a matching conflict. 
-				const conflict = this.migratePauseInfo_Helper_FindConflict(nodeData.conflicts, oldPauseInfo)
-
-				// If we can't find a matching conflict, let's just default to setting the node as blocked.
-				// Someone will have to fix this manually, so we might as well stop the entire node.
-				if (!conflict) {
-					this.settingsLogger.warn(`${nodeName} has 'branch-stop' but no matching conflict.`)
-					nodeData.pause.blockage = oldPauseInfo
-				}
-				// Now determine if this conflict is a node-stoppage conflict (syntax errors or no validtarget)
-				else if (conflict.kind === "Syntax error" || !conflict.targetBranchName ||
-					!nodeData.edges[conflict.targetBranchName]) {
-					nodeData.pause.blockage = oldPauseInfo
-				}
-				// If we don't need to stop the node for this, give it to the relevant edge
-				else {
-					this.settingsLogger.info(`\tApplying pause status to ${nodeName}->${conflict.targetBranchName}`)
-					nodeData.edges[conflict.targetBranchName].pause = {
-						blockage: oldPauseInfo
-					}
-				}
-			}
-
-			// Delete the old object
-			delete nodeData.pause.info
-		}
-		
-		// Also remove any old info string. The new class will generate appropriate values
-		if (nodeData.pause && nodeData.pause.infoStr) {
-			delete nodeData.pause.infoStr
-		}
-	}
-
-	private migratePauseInfo_Helper_FindConflict(nodeConflictData: any, oldBranchStopInfo: { [key: string]: any }): PersistentConflict | undefined {
-		if (!nodeConflictData || !Array.isArray(nodeConflictData)) {
-			return undefined
-		}
-
-		const conflictArray = nodeConflictData as PersistentConflict[]
-
-		if (conflictArray.length === 0) {
-			return undefined
-		}
-
-		const targetBranch = oldBranchStopInfo.targetBranchName ? 
-			oldBranchStopInfo.targetBranchName.toUpperCase() :
-			null
-
-		return conflictArray.find((conflict) => {
-			return conflict.targetBranchName === targetBranch && conflict.sourceCl === oldBranchStopInfo.sourceCl
-		})
-	}
-
-	private migratePauseState_Helper_runToJSON(nodeName: string, nodeData: {[key: string]: any}) {
-		if (nodeData.pause) {
-			if (nodeData.pause.blockage || nodeData.pause.manual_pause) {
-				// I think this is ok: it happens if the branch has not been paused since the 3.0.0 migration
-				this.settingsLogger.info(`${nodeName} already appears to use new format`)
-			}
-			else {
-				nodeData.pause.blockage = nodeData.pause.blockagePauseInfo
-				delete nodeData.pause.blockagePauseInfo
-				nodeData.pause.manual_pause = nodeData.pause.manualPauseInfo
-				delete nodeData.pause.manualPauseInfo
-			}
-		}
-	}
-
-	private removeAllSlackMessages() {
-		const notifications = this.object[NOTIFICATIONS_PERSISTENCE_KEY]
-
-		if (!notifications) {
-			this.settingsLogger.warn("Could not find notifications for Robomerge in settings file. Skipping...")
-			return
-		}
-
-		const slackMsgs = notifications[SLACK_MESSAGES_PERSISTENCE_KEY]
-
-		if (!slackMsgs) {
-			this.settingsLogger.warn("Could not find Slack messages in notifications section of settings file. Skipping...")
-			return
-		}
-
-		this.settingsLogger.info("Sunsetting all Slack messages to ensure a refresh of all notifications")
-		this.object[NOTIFICATIONS_PERSISTENCE_KEY][SLACK_MESSAGES_PERSISTENCE_KEY] = {}
-	}
-
 	getContext(name: string) {
 		return new Context(this, name);
 	}
