Remove unused init

diff --git a/Engine/Extras/RoboMerge/v3/src/robo/preview.ts b/Engine/Extras/RoboMerge/v3/src/robo/preview.ts
--- a/Engine/Extras/RoboMerge/v3/src/robo/preview.ts
+++ b/Engine/Extras/RoboMerge/v3/src/robo/preview.ts
@@ -6,11 +6,6 @@ import { Status } from './status';
 
 const logger = new ContextualLogger('preview')
 
-// let branchSpecsRootPath = ''
-export function init(_root: string) {
-	// branchSpecsRootPath = root
-}
-
 export async function getPreview(cl: number, singleBot?: string) {
 
 	let p4 = PerforceContext.getServerContext(logger)
diff --git a/Engine/Extras/RoboMerge/v3/src/robo/robo.ts b/Engine/Extras/RoboMerge/v3/src/robo/robo.ts
--- a/Engine/Extras/RoboMerge/v3/src/robo/robo.ts
+++ b/Engine/Extras/RoboMerge/v3/src/robo/robo.ts
@@ -26,7 +26,6 @@ import { notificationsInit } from './notifications'
 import { roboArgsInit } from './roboargs'
 import { settingsInit } from './settings'
 import { Vault } from './vault'
-import * as Preview from './preview'
 
 /*************************
  * RoboMerge main process
@@ -225,7 +224,6 @@ const vault = new Vault(args.vault)
 notificationsInit(args, vault)
 roboArgsInit(args)
 settingsInit(args)
-Preview.init(args.branchSpecsRootPath)
 
 const env = args.epicEnv || (args.nodeEnv === "production" ? "prod" : "dev")
 const sentryEnv = `${env}-${args.epicDeployment}`
