Move graph in to robo with all other files

diff --git a/Engine/Extras/RoboMerge/v3/src/new/graph.ts b/Engine/Extras/RoboMerge/v3/src/new/graph.ts
deleted file mode 100644
diff --git a/Engine/Extras/RoboMerge/v3/src/robo/graph.ts b/Engine/Extras/RoboMerge/v3/src/robo/graph.ts
new file mode 100644
--- /dev/null
+++ b/Engine/Extras/RoboMerge/v3/src/robo/graph.ts
@@ -0,0 +1,932 @@
+// Copyright Epic Games, Inc. All Rights Reserved.
+
+import { setDefault, sortBy } from '../common/helper';
+import { ContextualLogger } from '../common/logger';
+import { PerforceContext } from '../common/perforce';
+import { Branch, BranchGraphInterface } from '../robo/branch-interfaces';
+import { Status } from '../robo/status';
+
+export type TargetName = string & { __targetBrand: any }
+export type Stream = string & { __streamBrand: any }
+export type BotName = string & { __BotNameBrand: any }
+
+type ComputeResult = Map<Edge, Edge[]>
+
+type Success = 'failed' | 'succeeded'
+
+// deliberately not much in here - assuming stream for the moment
+export class Node {
+	constructor(public stream: Stream, public serverID?: string) // of form //x[/x]*
+	{}
+
+	offerDebugName(name: string) {
+		this._debugName = this._debugName || name
+	}
+
+	get debugName(): string {
+		return this._debugName || this.stream
+	}
+	
+	private _debugName: string | null = null
+}
+
+type EdgeFlag = 'automatic' | 'general'
+
+export class Edge
+{
+	constructor(
+		public source: Node,
+		public target: Node,
+
+		public targetName: string, // name of target in bot that owns the edge
+
+		public flags: Set<EdgeFlag>,
+
+		public bot: BotName,					// all edges will be bot-specific in practice, although maybe not necessary
+		public sourceAnnotation: any		// type specified by client code, e.g. NodeBot servicing this edge, or just last cl
+	) {}
+}
+
+// all aliases are of the form '<bot>:<name>', name lowercase
+export function makeTargetName(bot: BotName, nodeName: string) {
+	return `${bot}:${nodeName.toLowerCase()}` as TargetName
+}
+
+function makeBranchId(arg: string | Branch) {
+	let depotPath: string
+	if ((arg as Branch).uniqueBranch) {
+		return (arg as Branch).name as Stream
+	}
+	else if ((arg as Branch).rootPath) {
+		depotPath = (arg as Branch).rootPath.replace('/...', '')
+	}
+	else {
+		depotPath = arg as string
+	}
+
+	if (!depotPath.match(/^\/\/[\w-.][\/\w-.\*]*$/)) {
+		throw new Error('Invalid stream path: ' + depotPath)
+	}
+	return depotPath as Stream
+}
+
+/**
+
+What's the plan for bots with different branchspecs monitoring the same node?
+Doesn't matter: branchspecs are just an edge thing
+Nodes are uniquely identified by their rootPath/streamSubpath
+
+NB: one node in the ubergraph can represent multiple nodebots across bots
+
+ */ 
+
+export class GraphAPI {
+	constructor(public graph: Graph) {
+	}
+
+	reset(graph: Graph) {
+		this.graph = graph
+	}
+
+	getNode(name: TargetName) {
+		return this.graph.getNode(name)
+	}
+
+	getBranchGraph(name: BotName) {
+		return this.graph.branchGraphAliases.get(name)
+	}
+
+	findAllRoutesBetween(src: Node, target: Node, flags?: EdgeFlag[], skip?: Set<Node>) {
+		return this.graph.findAllRoutesBetween(src, target, flags, skip)
+	}
+
+	computeImplicitTargets(source: Node, requestedTargets: Map<Node, MergeMode>, allowedBots?: string[]) {
+		return this.graph.computeImplicitTargets(source, requestedTargets, allowedBots)
+	}
+}
+
+type EdgeDump = {
+	target: string
+	flags: string
+}
+
+type NodeDump = {
+	edges: EdgeDump[]
+	aliases: string[]
+}
+
+export class Graph {
+	private streamNodes = new Map<Stream, Node>()
+	private targetNames = new Map<TargetName, Node>()
+	private edgesBySource = new Map<Node, Set<Edge>>()
+	private edgesByTarget = new Map<Node, Set<Edge>>()
+	private edgesByBot = new Map<string, Set<Edge>>()
+
+	branchGraphAliases = new Map<string, BranchGraphInterface>()
+
+	findOrCreateStreamNode(depotPath: string | Branch) {
+		const branchId = makeBranchId(depotPath)
+		const server = (depotPath as Branch).config?.streamServer
+		return setDefault(this.streamNodes, branchId, new Node(branchId, server))
+	}
+
+	findNodeForStream(depotPath: string | Branch, allowPartialMatch?: boolean) {
+		const branchId = makeBranchId(depotPath)
+		let nodeForStream = this.streamNodes.get(branchId)
+		if (!nodeForStream && allowPartialMatch) {
+			const partialMatchbranchId = (branchId as string) + '/'
+			for (const [stream, node] of this.streamNodes) {
+				if (stream.startsWith(partialMatchbranchId)) {
+					nodeForStream = node
+					break
+				}
+			}
+		}
+		return nodeForStream
+	}
+
+	findNodeForBranch(branch: Branch) {
+		// need to standardise on rootpath/stream subpath (maybe rootpath is right?)
+		return this.findNodeForStream(branch.stream!)
+	}
+
+	addEdge(edge: Edge) {
+		setDefault(this.edgesBySource, edge.source, new Set).add(edge)
+		setDefault(this.edgesByTarget, edge.target, new Set).add(edge)
+		setDefault(this.edgesByBot, edge.bot, new Set).add(edge)
+	}
+
+	addNameForNode(node: Node, targetName: TargetName) {
+		// @todo check name is not already assigned to another node
+		this.targetNames.set(targetName, node)
+		node.offerDebugName(targetName)
+	}
+
+	getNode(name: TargetName) {
+		return this.targetNames.get(name)
+	}
+
+	private static filterEdgesByFlags(edges: Set<Edge>, flags: EdgeFlag[]) {
+		if (flags.length === 0)
+			return edges
+
+		return new Set<Edge>([...edges].filter(edge => {
+			for (const flag of flags) {
+				if (!edge.flags.has(flag))
+					return false
+			}
+			return true
+		}))
+	}
+
+	getEdgesBySource(src: Node, ...flags: EdgeFlag[]) {
+		return Graph.filterEdgesByFlags(this.edgesBySource.get(src) || new Set<Edge>, flags)
+	}
+
+	getEdgesByTarget(target: Node, ...flags: EdgeFlag[]) {
+		return Graph.filterEdgesByFlags(this.edgesByTarget.get(target) || new Set<Edge>, flags)
+	}
+
+	getEdgesForNode(node: Node, ...flags: EdgeFlag[]) {
+		const srcEdges = this.edgesBySource.get(node) || new Set<Edge>()
+		const targetEdges = this.edgesByTarget.get(node) || new Set<Edge>()
+		return Graph.filterEdgesByFlags(new Set([...srcEdges, ...targetEdges]), flags)
+	}
+
+	computeReachable(result: Set<Node>, src: Node, ...flags: EdgeFlag[]) {
+		// if result starts off with stuff, those nodes will be ignored (i.e. not checked)
+
+		for (const edge of this.getEdgesBySource(src, ...flags)) {
+			if (!result.has(edge.target)) {
+				result.add(edge.target)
+				this.computeReachable(result, edge.target, ...flags)
+			}
+		}
+	}
+
+	findRouteBetween(src: Node, target: Node, flags?: EdgeFlag[], skip?: Set<Node>) {
+		const seen = skip || new Set<Node>()
+		seen.add(src)
+		return this.findRouteBetweenImpl([], src, seen, target, flags || [])
+	}
+
+	findAllRoutesBetween(src: Node, target: Node, flags?: EdgeFlag[], skip?: Set<Node>) {
+		const seen = skip || new Set<Node>()
+		seen.add(src)
+		const result: Edge[][] = []
+		this.findAllRoutesBetweenImpl(result, [], src, seen, target, flags || [])
+		return result
+	}
+
+	dump(tags?: Set<string>, logger?: ContextualLogger) {
+
+		let validBots: string[] = []
+		if (tags) {
+			validBots = Array.from(this.branchGraphAliases)
+			                 .filter(([_,v]) => Status.includeBranch(v.config.visibility, tags, logger))
+			                 .map(([k,_]) => k)
+		}
+
+		const nodeNames = new Map<Node, TargetName[]>()
+		for (const [name, node] of this.targetNames) {
+			const botName = name.substring(0,name.search(':'))
+			if (!tags || validBots.includes(botName)) {
+				setDefault(nodeNames, node, []).push(name)
+			}
+		}
+
+		const orderedNodes = [...nodeNames.keys()]
+		sortBy(orderedNodes, node => node.debugName)
+
+		const result: [string, NodeDump][] = []
+		for (const node of orderedNodes) {
+			const info: NodeDump = {edges: [], aliases: nodeNames.get(node)!}
+			info.aliases.sort()
+
+			for (const edge of this.getEdgesBySource(node)) {
+				if (!tags || validBots.includes(edge.bot)) {
+					info.edges.push({target: edge.target.debugName, flags: [...edge.flags].join(', ')})
+				}
+			}
+			result.push([node.debugName, info])
+		}
+
+		return result
+	}
+
+	private findRouteBetweenImpl(routeSoFar: Edge[], src: Node, seen: Set<Node>, target: Node, flags: EdgeFlag[]): Edge[] | null {
+		for (const edge of this.getEdgesBySource(src, ...flags)) {
+			if (!seen.has(edge.target)) {
+				const routePlus = [...routeSoFar, edge]
+				if (edge.target === target)
+					return routePlus
+
+				seen.add(edge.target)
+				const route = this.findRouteBetweenImpl(routePlus, edge.target, seen, target, flags)
+				if (route)
+					return route
+			}
+		}
+		return null
+	}
+
+	private findAllRoutesBetweenImpl(completeRoutes: Edge[][], routeSoFar: Edge[],
+																	src: Node, seen: Set<Node>, target: Node, flags: EdgeFlag[]) {
+		for (const edge of this.getEdgesBySource(src, ...flags)) {
+			if (!seen.has(edge.target)) {
+				const routePlus = [...routeSoFar, edge]
+				if (edge.target === target) {
+					completeRoutes.push(routePlus)
+				}
+				else {
+					seen.add(edge.target)
+					this.findAllRoutesBetweenImpl(completeRoutes, routePlus, edge.target, seen, target, flags)
+					seen.delete(edge.target)
+				}
+			}
+		}
+	}
+
+	computeImplicitTargets(source: Node, requestedTargets: Map<Node, MergeMode>, allowedBots?: string[])
+		: {status: Success, integrations?: ComputeResult, unreachable?: Node[]} {
+
+		const targetsToFind = new Set<Node>()
+		const skipNodes = new Set<Node>()
+		for (const [target, mergeMode] of requestedTargets.entries()) {
+			if (mergeMode === 'skip') {
+				skipNodes.add(target)
+			}
+			else {
+				targetsToFind.add(target)
+			}
+		}
+
+		// initial edge to follow, then subsequent
+		const merges = new Map<Edge, Edge[]>()
+
+		// parallel flood the flowsTo graph to find all specified targets
+
+		// branchPaths is an array of 3-tuples, representing the branches at the current boundary of the flood
+		//	- first entry is the branch
+		//	- second and third entries make up the path to get to that branch
+		//		- second entry, if non-empty, has the path up to the most recent explicit target
+		//		- third entry has any remaining implicit targets
+
+		const branchPaths: [Node, Edge[], Edge[]][] = [[source, [], []]]
+
+		const seen = new Set([source])
+		while (targetsToFind.size !== 0 && branchPaths.length !== 0) {
+			const [sourceNode, explicit, implicit] = branchPaths.shift()!
+
+			// flood the graph one step to include all unseen direct flowsTo nodes of one branch
+			let anyUnseen = false
+
+			let flowsTo: Set<Edge> | Edge[] | undefined = this.edgesBySource.get(sourceNode)
+			if (flowsTo && allowedBots) {
+				flowsTo = [...flowsTo].filter((e: Edge) => allowedBots.indexOf(e.bot) >= 0)
+			}
+			if (flowsTo && !skipNodes.has(sourceNode)) {
+				for (const edge of flowsTo) {
+					const node = edge.target
+					if (seen.has(node))
+						continue
+
+					seen.add(node)
+					anyUnseen = true
+
+					if (targetsToFind.has(node)) {
+						branchPaths.unshift([node, [...explicit, ...implicit, edge], []])
+					}
+					else {
+						branchPaths.push([node, explicit, [...implicit, edge]])
+					}
+				}
+			}
+
+			// store the calculated merge data if we've exhausted a sub-tree of the graph
+			if (!anyUnseen && explicit.length > 0) {
+				const directEdge = explicit[0]
+				targetsToFind.delete(directEdge.target)
+
+				if (skipNodes.has(directEdge.target))
+					continue
+
+				const indirectTargets = setDefault(merges, directEdge, [])
+
+				for (let index = 1; index < explicit.length; ++index) {
+					const indirectEdge = explicit[index]
+					if (indirectTargets.indexOf(indirectEdge) < 0) {
+						indirectTargets.push(indirectEdge)
+					}
+					targetsToFind.delete(indirectEdge.target)
+				}
+			}
+		}
+
+		if (targetsToFind.size > 0) {
+			return {status: 'failed', unreachable: [...targetsToFind]}
+		}
+
+		return {status: 'succeeded', integrations: merges}
+	}
+
+}
+
+export function addBranchGraph(graph: Graph, branchGraph: BranchGraphInterface) {
+
+	graph.branchGraphAliases  = new Map([
+		[branchGraph.botname.toUpperCase(), branchGraph],
+		...graph.branchGraphAliases.entries(),
+		...(branchGraph.config.aliases.map(s => [s.toUpperCase(), branchGraph]) as [[string, BranchGraphInterface]])
+	])
+
+	const botname = branchGraph.botname as BotName
+	const branchNodes = new Map<Branch, Node>()
+
+	for (const branch of branchGraph.branches) {
+		if (!branch.bot) {
+			throw new Error(`branch ${branch.name} not running!`) // fine, but try again later
+		}
+		const branchNode = graph.findOrCreateStreamNode(branch)
+		branchNodes.set(branch, branchNode)
+		graph.addNameForNode(branchNode, makeTargetName(botname, branch.name))
+
+		for (const alias of branch.aliases) {
+			graph.addNameForNode(branchNode, makeTargetName(botname, alias))
+		}
+	}
+
+	for (const branch of branchGraph.branches) {
+		const sourceNode = branchNodes.get(branch)!
+		for (const target of branch.flowsTo) {
+			const flags = new Set<EdgeFlag>()
+			if (branch.forceFlowTo.indexOf(target) !== -1) {
+				flags.add('automatic')
+			}
+
+			if (branch.isDefaultBot) {
+				flags.add('general')
+			}
+
+			const targetBranch = branchGraph.getBranch(target)!
+			graph.addEdge(new Edge(sourceNode,
+				branchNodes.get(targetBranch)!,
+				targetBranch.name, flags, botname, branch.bot!
+			))
+		}
+	}
+
+	// if (branchNodes.size > 1) {
+	// 	//testing syntax
+	// 	const nodeIter = branchNodes.values()
+
+	// 	const src = nodeIter.next().value
+	// 	const targets = new Map<Node, MergeMode>()
+	// 	targets.set(nodeIter.next().value, 'normal')
+
+	// 	graph.computeImplicitTargets(src, targets)
+	// }
+}
+
+export type MergeMode = 'safe' | 'normal' | 'null' | 'clobber' | 'skip'
+
+class Test {
+	graph = new Graph
+	targets = new Map<Node, MergeMode>();
+
+	static readonly BOT_NAME = 'TEST' as BotName
+
+	constructor(nodeStrs: string[]) {
+
+		for (let index = 0; index < nodeStrs.length; ++index) {
+			const name = String.fromCharCode('a'.charCodeAt(0) + index)
+
+			const fromNode = this.addDummyStream(name)
+
+			for (const char of nodeStrs[index]) {
+				const targetNode = this.addDummyStream(char)
+
+				const flags = new Set<EdgeFlag>()
+				if (char.toUpperCase() === char) {
+					flags.add('automatic')
+				}
+
+				this.graph.addEdge(new Edge(fromNode, targetNode, char, flags, Test.BOT_NAME, {}))
+			}
+		}
+	}
+
+	computeTargets(source: string, targetsString: string) {
+		const sourceNode = this.getNode(source)
+		this.fillTargetsMap(targetsString)
+
+		return this.graph.computeImplicitTargets(sourceNode, this.targets)
+	}
+
+	formatTestComputeTargetsResult(computed: ComputeResult) {
+		const bits: string[] = []
+		for (const [direct, indirect] of computed.entries()) {
+			bits.push(this.formatNode(direct.target) + ': ' + indirect.map(x => this.formatNode(x.target)).join(', '))
+		}
+		return bits.join('; ')
+	}
+
+	getNode(name: string) {
+		const node = this.graph.getNode(makeTargetName(Test.BOT_NAME, name))
+		if (!node) throw 'arse!'
+		return node
+	}
+
+	private fillTargetsMap(targets: string) {
+
+		const targetBits = targets.split(':') // up to three section in target (1st) string: normal, skip, null
+
+		for (const targetChar of targetBits[0]) {
+			this.targets.set(this.getNode(targetChar), 'normal')
+		}
+
+		if (targetBits.length > 1) {
+			for (const targetChar of targetBits[1]) {
+				this.targets.set(this.getNode(targetChar), 'skip')
+			}
+		}
+
+		if (targetBits.length > 2) {
+			for (const targetChar of targetBits[2]) {
+				this.targets.set(this.getNode(targetChar), 'null')
+			}
+		}
+	}
+
+	private formatNode(node: Node) {
+
+		const mm = this.targets.get(node)
+		const prefix = !mm || mm === 'normal' ? '' :
+			mm === 'skip' ? '-' :
+			mm === 'null' ? '!' : '@'
+
+		return prefix + node.debugName[5]
+	}
+
+	private addDummyStream(name: string) {
+		const targetName = makeTargetName(Test.BOT_NAME, name)
+		let node = this.graph.getNode(targetName)
+		if (!node) {
+			node = this.graph.findOrCreateStreamNode('//Test/' + name.toUpperCase())
+			this.graph.addNameForNode(node, targetName)
+		}
+		return node
+	}
+
+}
+
+// async function inner() {
+// 	throw new Error('test')
+// }
+
+// function proxy() {
+// 	const promise = inner()
+// 	promise.finally(() => console.log('here!'))
+// 	return promise
+// }
+
+// async function outer() {
+// 	try {
+// 		await proxy()
+// 	}
+// 	catch (err) {
+// 		console.log('caught!')
+// 	}
+// }
+
+
+// function runUnhandledRejectionTest() {
+// 	outer()
+// }
+
+//  _______        _       
+// |__   __|      | |      
+//    | | ___  ___| |_ ___ 
+//    | |/ _ \/ __| __/ __|
+//    | |  __/\__ \ |_\__ \
+//    |_|\___||___/\__|___/
+
+// for unit tests!
+
+
+
+const colors = require('colors')
+colors.enable()
+colors.setTheme(require('colors/themes/generic-logging.js'))
+
+
+export function runTests(parentLogger: ContextualLogger) {
+
+
+
+	// runUnhandledRejectionTest()
+
+/**
+
+Test format:
+
+[[command, node, ...<graph definition>], expected]
+
+graph definition e.g.: ['b', 'c', ''] means a->b, b->c and d also exists
+expected: ;-separated <direct>:<indirect>
+
+1: A -> B -> C		:c expected to produce '' (c unreachable, could alternatively be syntax error)
+2: A => B -> C		c:b expected to produce null (c unreachable)
+3: A -> B -> C	_:c should produce '' (skip shouldn't affect route)
+
+4/5: dev/release set-up
+
+
+6: A -> B -> c		c:d expected to produce B:C (A->B->C)
+	 -> D -> c
+
+7: A -> B -> c		c:b expected to produce D:C (A->D->C)
+	 => D -> c
+
+8: A => B => C		_:c expected to produce ''
+	 -> D
+
+9: A => B => C		_:c:b expected to produce !B
+
+
+Note first level is treated specially: for the purposes of these tests, force flow from the start node is effectively ignored
+*/
+
+
+// want to find a way of testing commented out ones - the relevant skip handling is done in targets.ts
+
+
+	const unitTestLogger = parentLogger.createChild('Graph')
+	const tests: [string[], string | null][] = [
+
+		[[':c', 'a',		'b', 'c', ''],										''],
+		[['c:b', 'a', 		'b', 'c', ''],										null],
+		[['c:b', 'a', 		'B', 'c', ''],										null],
+		[[':c', 'a', 		'b', 'c', ''],										''],
+
+		[['c', 'a',			'bDE', 'Ac', 'B', 'a', 'a'],						'B:C'],
+		[[':c', 'a',		'bDE', 'Ac', 'B', 'a', 'a'],						''], //DE'], // 5
+
+		[['c:d', 'a', 		'bd', 'c', '', 'c'],								'B:C'],
+		[['c:b', 'a', 		'bD', 'c', '', 'c'],								'D:C'],
+		[[':c', 'a',		'Bd', 'C', '', ''],									''],
+		[[':c:b', 'a',		'B', 'C', ''],										'!B:'], // !B:-C'], skip is checked by functional test
+		[['b', 'a', 		'b', 'C', ''],										'B:'], // 10
+
+// usual dev/release set-up (d-g are releases going back in time)
+		[['f', 'b',			'bcd', 'a', 'a', 'Ae', 'Df', 'Eg', 'F'],			'A:DEF'],
+		[['b', 'f',			'bcd', 'a', 'a', 'Ae', 'Df', 'Eg', 'F'],			'E:DAB'],
+		[['a', 'g',			'bcd', 'a', 'a', 'Ae', 'Df', 'Eg', 'F'],			'F:EDA'],
+		[[':a', 'g',		'bcd', 'a', 'a', 'Ae', 'Df', 'Eg', 'F'],			''], // taking a look at this one
+
+		[['de', 'a',		'b', 'c', 'de', '', ''],							'B:CED'],
+		[['de', 'h',		'hc', 'hd', 'deFg', 'hbc', 'c', 'c', 'c', 'abd'],	'D:CE'],
+		[['de', 'h',		'bd', 'hc', 'hd', 'eFg', 'hbc', 'c', 'c', 'c'],		'C:DE'],
+		[['db', 'h',		'abd', 'hc', 'hd', 'deFg', 'hbc', 'c', 'c', 'c'],	'C:DEB'],
+		[['cg', 'a',		'be', 'acd', 'b', 'b', 'aFg', 'e', 'e'],			'B:C;E:G'],
+		[['cf', 'a', 		'bE', 'acd', 'b', 'b', 'aFg', 'e', 'e'],			'B:C;E:F'],
+		[['c:e', 'a', 		'bE', 'acd', 'b', 'b', 'aFg', 'e', 'e'],			'B:C'],
+		[['cf:b', 'a', 		'bE', 'acd', 'b', 'b', 'aFg', 'e', 'e'],			null],
+		[['fg', 'h',		'abd', 'hC', 'hd', 'defg', 'hbc', 'c', 'c', 'c'],	'C:DGF'],
+		[['dge', 'h',		'abd', 'hC', 'hd', 'defg', 'hbc', 'c', 'c', 'c'],	'C:DGE'],
+		[['e', 'i',			'abcd', 'Ie', 'if', 'iG', 'ih', 'a', 'b', 'c', 'D'],'D:GBE']
+	]
+
+	let success = 0, fail = 0, ran = 0
+	for (const [testStr, expected] of tests /*/ .slice(5, 6) /**/) {
+		const test = new Test(testStr.slice(2))
+		const result = test.computeTargets(testStr[1], testStr[0])
+
+		let expectedOnFail: string | null = null
+
+		const succeeded = result.status === 'succeeded'
+
+		const formattedResult = succeeded ? test.formatTestComputeTargetsResult(result.integrations!) : 'failed'
+		if (expected !== null) {
+			// not expecting an error
+			if (!succeeded) {
+				// store string we were expecting
+				expectedOnFail = expected
+			}
+			else if (formattedResult.toUpperCase().replace(/\s|,/g, '') !== expected) {
+				expectedOnFail = expected
+			}
+		}
+		else if (succeeded) {
+			expectedOnFail = 'fail'
+		}
+
+		++ran
+		if (expectedOnFail === undefined) throw new Error('wut')
+		if (expectedOnFail !== null) {
+			++fail
+			const EMPTY = '<empty>'
+			unitTestLogger.warn(`Test ${colors.warn(ran.toString().padStart(2))} failed:   ` +
+				`${colors.warn((formattedResult || EMPTY).padStart(10))} vs ${colors.warn(expectedOnFail || EMPTY)}`)
+		}
+		else {
+			++success
+		}
+	}
+	const message = `Graph tests: ran ${ran} tests, ${success} matched, ${fail} failed`
+	unitTestLogger.info((fail > 0 ? colors.error : colors.info)(message))
+	return fail
+}
+
+type ChangeFlag = 'manual' | 'null' | 'ignore' | 'disregardexcludedauthors' | 'roboshelf' | 'disregardassetblock'
+const ALLOWED_RAW_FLAGS = new Set(['null','ignore','deadend'])
+
+const FLAGMAP: {[name: string]: ChangeFlag} = {
+	ignore: 'ignore'
+};
+
+// should probably go in NodeBot, although specifically only dealing with string parsing
+
+// gets tricky straight away dealing with commands for other bots
+// will make target names here
+function parseTargetsAndFlagsImpl(tokens: string[], logger: ContextualLogger, forcedMode?: MergeMode) {
+	const flags = new Set<ChangeFlag>()
+	const targets = new Map<string, MergeMode>()
+
+	for (const token of tokens) {
+		const tokenLower = token.toLowerCase()
+
+		// check for 'ignore' and remap to #ignore flag
+		if (ALLOWED_RAW_FLAGS.has(tokenLower))
+		{
+			flags.add(FLAGMAP[tokenLower]!);
+			continue;
+		}
+
+		// see if this has a modifier
+		const prefix = token.charAt(0)
+		if (prefix === '#' || prefix === '$') { // $ provided as alternative, but not used as far as I know
+			// set the flag and continue the loop
+			const flagname = FLAGMAP[tokenLower.substring(1)]
+			if (flagname) {
+				flags.add(flagname)
+			}
+			else {
+				logger.info(`Ignoring unknown flag "${token}"`)
+			}
+			continue
+		}
+
+		let targetName = token
+		let mergeMode = forcedMode || 'normal'
+		if (prefix === '!' || prefix === '-')
+		{
+			targetName = token.substring(1)
+			mergeMode = prefix === '!' ? 'null' : 'skip'
+		}
+
+		const existingMergeMode = targets.get(targetName)
+		if (existingMergeMode) {
+			if (existingMergeMode !== mergeMode) {
+				throw new Error('mismatch!') // todo add to list of errors instead
+			}
+		}
+
+		targets.set(targetName, mergeMode)
+	}
+
+	return {status: 'succeeded', flags, targets}
+}
+
+export function parseTargetsAndFlags(graph: Graph, botname: BotName, tokens: string[], automaticTargets: Node[],
+																logger: ContextualLogger, forcedMode?: MergeMode) {
+	const result = parseTargetsAndFlagsImpl(tokens, logger, forcedMode)
+	if (result.status !== 'succeeded') {
+		throw new Error('to do!')
+	}
+
+	const targets = new Map<Node, MergeMode>()
+
+	// make sure we can find a node for each target
+	for (const [targetStr, mergeMode] of result.targets) {
+		const node = graph.getNode(makeTargetName(botname, targetStr))
+		if (!node) {
+			throw new Error('to do!')
+		}
+		targets.set(node, mergeMode)
+	}
+
+
+// kind of odd to do this before computing implicit, but shouldn't do any harm
+	for (const auto of automaticTargets) {
+		if (!targets.has(auto)) {
+			targets.set(auto, forcedMode ?? 'normal')
+		}
+	}
+
+	return {status: 'succeeded', flags: result.flags, targets}
+}
+
+type StringBotNamePair = [string, BotName | 'default']
+
+function parseChange(cldesc: string) {
+
+	// parse the description
+	const descLines = cldesc.split('\n')
+
+	let tokens: StringBotNamePair[] = []
+	for (const line of descLines) {
+		// trim end - keep any initial whitespace
+		const comp = line.replace(/\s+$/, '')
+
+		// check for control hashes
+		const match = comp.match(/^(\s*)#([-\w[\]]+)[:\s]*(.*)$/)
+		if (!match)
+			continue
+
+		const ws = match[1] 
+		const command = match[2].toUpperCase()
+		const value = match[3].trim()
+
+		// #robomerge tags are required to be at the start of the line
+		if (ws)
+			// could do this by changing regex, but keeping this to match flow in NodeBot code
+			continue
+
+		let bot: BotName | 'default' | null = null
+	 	if (command === 'ROBOMERGE') {
+
+			// completely ignore bare ROBOMERGE tags if we're not the default bot (should not affect default flow)
+			bot = 'default'
+		}
+
+		else  {
+			// what sort of robomerge command is it
+			const specificMatch = command.match(/^ROBOMERGE\[(.*)\]$/)
+			if (specificMatch) {
+				bot = specificMatch[1] as BotName
+			}
+		}
+
+		if (bot) {
+			const tokensForThisLine: StringBotNamePair[] = value.split(/[ ,]/)
+													.filter(Boolean).map(token => [token, bot] as StringBotNamePair)
+			tokens = [...tokens, ...tokensForThisLine]
+		}
+	}
+
+	return tokens
+}
+
+export class Trace {
+	private readonly traceLogger: ContextualLogger
+	private readonly p4: PerforceContext
+	constructor(private graph: Graph, parentLogger: ContextualLogger)
+	{
+		this.traceLogger = parentLogger.createChild('Trace')
+		this.p4 = PerforceContext.getServerContext(this.traceLogger)
+	}
+
+	async traceRoute(cl: number, sourceNodeStr: string, targetNodeStr: string) {
+		// const sourceBranch = this.resolveBranch(sourceBranchStr)
+		const sourceNode = this.graph.findNodeForStream(sourceNodeStr)
+		const targetNode = this.graph.findNodeForStream(targetNodeStr)
+
+// should just dump source if source and target the same?
+
+		if (!sourceNode)
+			return 'UNKNOWN_SOURCE_BRANCH'
+
+		if (!targetNode)
+			return 'UNKNOWN_TARGET_BRANCH'
+
+		if (sourceNode === targetNode)
+			return 'SOURCE_AND_TARGET_THE_SAME'
+
+		let change
+		try {
+			change = await this.p4.getChange(cl)
+		}
+		catch (err) {
+			return 'CHANGELIST_NOT_FOUND'
+		}
+
+
+		if (typeof change.desc !== 'string')
+			return 'INVALID_CHANGELIST_DESCRIPTION'
+
+		const tokensAndBots = parseChange(change.desc)
+
+		const outEdges = this.graph.getEdgesBySource(sourceNode)
+		if (outEdges.size === 0)
+			return 'CHANGE_WILL_NOT_REACH_TARGET'
+
+		// will have to check if the source node is served by multiple bots, and if so, calculate targets for each of them
+		// for now, just pick one
+		const arbitraryEdge = outEdges.values().next().value
+		const botname = arbitraryEdge.bot
+		const isDefaultBot = arbitraryEdge.flags.has('general')
+
+		const relevantTokens = tokensAndBots
+			.filter(([_, bot]) => (bot === botname || isDefaultBot && bot === 'default'))
+			.map(([token, _]) => token)
+
+		const parseResult = parseTargetsAndFlags(this.graph, botname, relevantTokens, 
+			[/*...graph.getEdgesBySource(sourceNode, 'automatic')*/], this.traceLogger)
+
+		if (parseResult.status !== 'succeeded')
+			return 'INVALID_CHANGELIST_DESCRIPTION'
+
+		const skipNodes = new Set<Node>([...parseResult.targets.entries()]
+			.filter(([_, mergeMode]) => mergeMode === 'skip')
+			.map(([node, _]) => node)
+		);
+
+		// find a route (would ideally be able to cope with multiple valid routes)
+		const route = this.graph.findRouteBetween(sourceNode, targetNode, [], skipNodes)
+		if (!route)
+			return 'CHANGE_WILL_NOT_REACH_TARGET'
+
+		// - find last non-automatic edge, if any: its target needs to be somewhere in resultant merges
+		let nonAutoEdge: Edge | null = null
+		for (let index = route.length - 1; index >= 0; --index) {
+			const edge = route[index]
+			if (!edge.flags.has('automatic')) {
+				nonAutoEdge = edge
+				break
+			}
+		}
+
+		// always compute targets in case it fails
+		const computeResult = this.graph.computeImplicitTargets(sourceNode, parseResult.targets)
+
+		if (computeResult.status !== 'succeeded')
+			return 'INVALID_CHANGELIST_DESCRIPTION'
+
+
+		if (nonAutoEdge) {
+
+			let found = false
+
+		outer:
+			for (const [firstEdge, furtherEdges] of computeResult.integrations!) {
+				if (firstEdge.target === nonAutoEdge.target) {
+					found = true
+					break
+				}
+
+				for (const edge of furtherEdges) {
+					if (edge.target === nonAutoEdge.target) {
+						found = true
+						break outer
+					}
+				}
+			}
+
+			if (!found) 
+				return 'CHANGE_WILL_NOT_REACH_TARGET'
+		}
+
+		return route
+	}
+}
diff --git a/Engine/Extras/RoboMerge/v3/src/robo/graphbot.ts b/Engine/Extras/RoboMerge/v3/src/robo/graphbot.ts
--- a/Engine/Extras/RoboMerge/v3/src/robo/graphbot.ts
+++ b/Engine/Extras/RoboMerge/v3/src/robo/graphbot.ts
@@ -21,7 +21,7 @@ import { Settings } from './settings';
 import { Status } from './status';
 import { GraphBotState } from "./status-types"
 import { TickJournal } from './tick-journal';
-import { GraphAPI } from '../new/graph';
+import { GraphAPI } from './graph';
 import * as p4util from '../common/p4util';
 
 // probably get the gist after 2000 characters
diff --git a/Engine/Extras/RoboMerge/v3/src/robo/ipc.ts b/Engine/Extras/RoboMerge/v3/src/robo/ipc.ts
--- a/Engine/Extras/RoboMerge/v3/src/robo/ipc.ts
+++ b/Engine/Extras/RoboMerge/v3/src/robo/ipc.ts
@@ -2,7 +2,7 @@
 
 import { ContextualLogger, isNpmLogLevel, NpmLogLevelCompare, NpmLogLevelValues } from '../common/logger';
 import { getRunningPerforceCommands, PerforceContext } from '../common/perforce';
-import { Trace } from '../new/graph';
+import { Trace } from './graph';
 import { IPCControls, EdgeBotInterface, NodeBotInterface } from './bot-interfaces';
 import { OperationResult } from './branch-interfaces';
 import { GraphBot } from './graphbot';
diff --git a/Engine/Extras/RoboMerge/v3/src/robo/nodebot.ts b/Engine/Extras/RoboMerge/v3/src/robo/nodebot.ts
--- a/Engine/Extras/RoboMerge/v3/src/robo/nodebot.ts
+++ b/Engine/Extras/RoboMerge/v3/src/robo/nodebot.ts
@@ -24,7 +24,7 @@ import { SlackFile, SlackMessage, SlackMessageField, SlackMessageStyles } from '
 import { PauseState } from './state-interfaces';
 import { newTickJournal, TickJournal } from './tick-journal';
 import { computeTargets, descriptionParsers, parseDescription, processOtherBotTargets, getIntegrationOwner, getNodeBotFullName, getNodeBotFullNameForLogging } from './targets';
-import { GraphAPI } from '../new/graph';
+import { GraphAPI } from './graph';
 
 /**********************************
  * Bot monitoring a single stream
diff --git a/Engine/Extras/RoboMerge/v3/src/robo/robo.ts b/Engine/Extras/RoboMerge/v3/src/robo/robo.ts
--- a/Engine/Extras/RoboMerge/v3/src/robo/robo.ts
+++ b/Engine/Extras/RoboMerge/v3/src/robo/robo.ts
@@ -13,7 +13,7 @@ import { Mailer } from '../common/mailer';
 import { initializePerforce, PerforceContext } from '../common/perforce';
 import { BuildVersion, VersionReader } from '../common/version';
 import { CertFiles } from '../common/webserver';
-import { addBranchGraph, Graph, GraphAPI } from '../new/graph';
+import { addBranchGraph, Graph, GraphAPI } from './graph';
 import { AutoBranchUpdater } from './autobranchupdater';
 import { Branch } from './branch-interfaces';
 import { Gate } from './gate';
diff --git a/Engine/Extras/RoboMerge/v3/src/robo/targets.ts b/Engine/Extras/RoboMerge/v3/src/robo/targets.ts
--- a/Engine/Extras/RoboMerge/v3/src/robo/targets.ts
+++ b/Engine/Extras/RoboMerge/v3/src/robo/targets.ts
@@ -2,7 +2,7 @@
 
 import { Branch, MergeAction, MergeMode } from './branch-interfaces';
 import { BranchGraphInterface, PendingChange, Target, TargetInfo } from './branch-interfaces';
-import { BotName, Edge, GraphAPI, Node, makeTargetName } from '../new/graph';
+import { BotName, Edge, GraphAPI, Node, makeTargetName } from './graph';
 import { ContextualLogger } from '../common/logger';
 import { Change } from '../common/perforce';
 
diff --git a/Engine/Extras/RoboMerge/v3/src/robo/trackchange.ts b/Engine/Extras/RoboMerge/v3/src/robo/trackchange.ts
--- a/Engine/Extras/RoboMerge/v3/src/robo/trackchange.ts
+++ b/Engine/Extras/RoboMerge/v3/src/robo/trackchange.ts
@@ -2,7 +2,7 @@
 
 import { ContextualLogger } from '../common/logger';
 import { DescribeResult, isExecP4Error, PerforceContext } from '../common/perforce';
-import { Node } from '../new/graph';
+import { Node } from './graph';
 import { NodeBotInterface } from './bot-interfaces';
 import { RoboMerge } from './robo';
 import { Status } from './status';
diff --git a/Engine/Extras/RoboMerge/v3/src/robo/unit-tests.ts b/Engine/Extras/RoboMerge/v3/src/robo/unit-tests.ts
--- a/Engine/Extras/RoboMerge/v3/src/robo/unit-tests.ts
+++ b/Engine/Extras/RoboMerge/v3/src/robo/unit-tests.ts
@@ -2,7 +2,7 @@
 
 import { ContextualLogger } from '../common/logger';
 import { runTests as gate_runTests } from './gate';
-import { runTests as graph_runTests } from '../new/graph';
+import { runTests as graph_runTests } from './graph';
 import { runTests as notifications_runTests } from './notifications';
 import { runTests as targets_runTests } from './targets';
 
