Update dependencies
No longer use deprecated request
Use jsonlint-mod instead of jsonlint

diff --git a/Engine/Extras/RoboMerge/v3/functional_tests/framework/package.json b/Engine/Extras/RoboMerge/v3/functional_tests/framework/package.json
--- a/Engine/Extras/RoboMerge/v3/functional_tests/framework/package.json
+++ b/Engine/Extras/RoboMerge/v3/functional_tests/framework/package.json
@@ -15,7 +15,6 @@
     "@types/node": "^22.9.0",
     "bent": "^7.1.2",
     "colors": "^1.4.0",
-    "jsonlint": "^1.6.3",
     "mustache": "^3.2.1"
   }
 }
diff --git a/Engine/Extras/RoboMerge/v3/functional_tests/framework/src/tests/confirm-binary-stomp.ts b/Engine/Extras/RoboMerge/v3/functional_tests/framework/src/tests/confirm-binary-stomp.ts
--- a/Engine/Extras/RoboMerge/v3/functional_tests/framework/src/tests/confirm-binary-stomp.ts
+++ b/Engine/Extras/RoboMerge/v3/functional_tests/framework/src/tests/confirm-binary-stomp.ts
@@ -1,8 +1,6 @@
 // Copyright Epic Games, Inc. All Rights Reserved.
 import { FunctionalTest, P4Client, P4Util, RobomergeBranchSpec, Stream } from '../framework'
 
-// const jsonlint: any = require('jsonlint')
-
 export const streams: Stream[] = [
 	{name: 'Main', streamType: 'mainline'},
 	{name: 'DevBranch1', streamType: 'development', parent: 'Main'},
diff --git a/Engine/Extras/RoboMerge/v3/functional_tests/framework/src/tests/stomp-with-add.ts b/Engine/Extras/RoboMerge/v3/functional_tests/framework/src/tests/stomp-with-add.ts
--- a/Engine/Extras/RoboMerge/v3/functional_tests/framework/src/tests/stomp-with-add.ts
+++ b/Engine/Extras/RoboMerge/v3/functional_tests/framework/src/tests/stomp-with-add.ts
@@ -1,8 +1,6 @@
 // Copyright Epic Games, Inc. All Rights Reserved.
 import { FunctionalTest, P4Util, RobomergeBranchSpec, Stream } from '../framework'
 
-// const jsonlint: any = require('jsonlint')
-
 export const streams: Stream[] = [
 	{name: 'Main', streamType: 'mainline'},
 	{name: 'DevParent', streamType: 'development', parent: 'Main'},
diff --git a/Engine/Extras/RoboMerge/v3/package.json b/Engine/Extras/RoboMerge/v3/package.json
--- a/Engine/Extras/RoboMerge/v3/package.json
+++ b/Engine/Extras/RoboMerge/v3/package.json
@@ -24,15 +24,13 @@
     "@types/swagger-ui-dist": "^3.0.3",
     "@viz-js/viz": "^3.11.0",
     "colors": "^1.4.0",
-    "jsonlint": "^1.6.3",
+    "jsonlint-mod": "^1.7.6",
     "ldapauth-fork": "^4.0.2",
     "markdown-toc": "^1.2.0",
     "marked": "^4.0.12",
     "mustache": "^3.2.1",
     "node-gzip": "^1.1.2",
     "nodemailer": "^6.7.0",
-    "nomnom": "^1.5.3",
-    "request": "^2.88.0",
     "sanitize-html": "^2.5.3",
     "semver": "^6.1.0",
     "source-map-support": "^0.5.12",
diff --git a/Engine/Extras/RoboMerge/v3/src/common/analytics.ts b/Engine/Extras/RoboMerge/v3/src/common/analytics.ts
--- a/Engine/Extras/RoboMerge/v3/src/common/analytics.ts
+++ b/Engine/Extras/RoboMerge/v3/src/common/analytics.ts
@@ -2,7 +2,6 @@
 
 import { ANALYTICS_SIMPLE_COUNTERS, newTickJournal, SimpleCounters, TickJournal } from '../robo/tick-journal';
 import { ContextualLogger } from './logger';
-import * as request from './request';
 
 const INFLUX_URL = '<influx end-point goes here>';
 const MEASUREMENT = 'robomerge';
@@ -204,7 +203,10 @@ export class Analytics {
 			return
 		}
 
-		request.post({url: INFLUX_URL, body})
+		fetch(INFLUX_URL, {
+			method: 'POST',
+			body
+		})
 		.catch(error => this.analyticsLogger!.printException(error, 'Analytics error'));
 	}
 
diff --git a/Engine/Extras/RoboMerge/v3/src/robo/branchdefs.ts b/Engine/Extras/RoboMerge/v3/src/robo/branchdefs.ts
--- a/Engine/Extras/RoboMerge/v3/src/robo/branchdefs.ts
+++ b/Engine/Extras/RoboMerge/v3/src/robo/branchdefs.ts
@@ -3,7 +3,7 @@
 import { PerforceContext } from '../common/perforce'
 import { ContextualLogger } from '../common/logger';
 
-const jsonlint: any = require('jsonlint')
+const jsonlint: any = require('jsonlint-mod')
 
 const RESERVED_BRANCH_NAMES = ['NONE', 'DEFAULT', 'IGNORE', 'DEADEND', ''];
 
diff --git a/Engine/Extras/RoboMerge/v3/src/robo/settings.ts b/Engine/Extras/RoboMerge/v3/src/robo/settings.ts
--- a/Engine/Extras/RoboMerge/v3/src/robo/settings.ts
+++ b/Engine/Extras/RoboMerge/v3/src/robo/settings.ts
@@ -25,7 +25,7 @@ export function settingsInit(inArgs: Args) {
 }
 	
 
-const jsonlint: any = require('jsonlint')
+const jsonlint: any = require('jsonlint-mod')
 
 function readFileToString(filename: string) {
 	try {
