Make RoboArgs a class instead of loose data

diff --git a/Engine/Extras/RoboMerge/v3/src/robo/edgebot.ts b/Engine/Extras/RoboMerge/v3/src/robo/edgebot.ts
--- a/Engine/Extras/RoboMerge/v3/src/robo/edgebot.ts
+++ b/Engine/Extras/RoboMerge/v3/src/robo/edgebot.ts
@@ -16,7 +16,7 @@ import { Gate } from "./gate";
 import { NodeBot } from "./nodebot";
 import { isUserAKnownBot, postMessageToChannel } from "./notifications";
 import { PerforceStatefulBot } from "./perforce-stateful-bot";
-import { getExclusiveLockOpenedsToRun } from "./roboargs";
+import { RoboArgs } from "./roboargs";
 import { Context } from "./settings";
 import { SlackMessageStyles } from "./slack";
 import { PauseState } from "./state-interfaces";
@@ -267,7 +267,7 @@ class EdgeBotImpl extends PerforceStatefulBot {
 		for (const err of errors) {
 			const match = err.match(EXCLUSIVE_CHECKOUT_REGEX)
 			if (match) {
-				if (paths.length < getExclusiveLockOpenedsToRun()) {
+				if (paths.length < RoboArgs.getExclusiveLockOpenedsToRun()) {
 					paths.push({depotPath: match[1] + match[2], name: match[2]})
 				} else {
 					results.push({depotPath: match[1] + match[2], name: match[2], user: "", client: ""})
diff --git a/Engine/Extras/RoboMerge/v3/src/robo/nodebot.ts b/Engine/Extras/RoboMerge/v3/src/robo/nodebot.ts
--- a/Engine/Extras/RoboMerge/v3/src/robo/nodebot.ts
+++ b/Engine/Extras/RoboMerge/v3/src/robo/nodebot.ts
@@ -17,7 +17,7 @@ import { BotEventTriggers } from './events';
 import { GraphBot } from './graphbot';
 import { getPingTarget, isUserSlackGroup, makeClLink, SlackMessages } from './notifications';
 import { PerforceStatefulBot } from './perforce-stateful-bot';
-import { getExclusiveLockOpenedsToRun } from './roboargs';
+import { RoboArgs } from './roboargs';
 import { BlockageNodeOpUrls, OperationUrlHelper } from './roboserver';
 import { Context } from './settings';
 import { SlackFile, SlackMessage, SlackMessageField, SlackMessageStyles } from './slack';
@@ -1088,8 +1088,8 @@ export class NodeBot extends PerforceStatefulBot implements NodeBotInterface {
 
 		// Verify that we have the client/user information for all locked files
 		let lockedFiles: ExclusiveFileDetails[] = blockageInfo.additionalInfo.exclusiveFiles
-		if (lockedFiles.length > getExclusiveLockOpenedsToRun()) {
-			lockedFiles = [...lockedFiles.slice(0,getExclusiveLockOpenedsToRun()), ...(await this.p4.getOpenedDetails(lockedFiles.slice(getExclusiveLockOpenedsToRun())))]
+		if (lockedFiles.length > RoboArgs.getExclusiveLockOpenedsToRun()) {
+			lockedFiles = [...lockedFiles.slice(0,RoboArgs.getExclusiveLockOpenedsToRun()), ...(await this.p4.getOpenedDetails(lockedFiles.slice(RoboArgs.getExclusiveLockOpenedsToRun())))]
 		}
 
 		const unknownClients = lockedFiles.filter((file) => file.client.length == 0)
diff --git a/Engine/Extras/RoboMerge/v3/src/robo/robo.ts b/Engine/Extras/RoboMerge/v3/src/robo/robo.ts
--- a/Engine/Extras/RoboMerge/v3/src/robo/robo.ts
+++ b/Engine/Extras/RoboMerge/v3/src/robo/robo.ts
@@ -23,7 +23,7 @@ import { NodeBot } from './nodebot';
 import { roboAnalytics, setGlobalAnalytics } from './roboanalytics';
 import { RoboServer } from './roboserver';
 import { notificationsInit } from './notifications'
-import { roboArgsInit } from './roboargs'
+import { RoboArgs } from './roboargs'
 import { settingsInit } from './settings'
 import { Vault } from './vault'
 
@@ -219,10 +219,10 @@ if (!maybeNullArgs) {
 	process.exit(1)
 }
 const args = maybeNullArgs!
+RoboArgs.Init(args)
 
 const vault = new Vault(args.vault)
 notificationsInit(args, vault)
-roboArgsInit(args)
 settingsInit(args)
 
 const env = args.epicEnv || (args.nodeEnv === "production" ? "prod" : "dev")
diff --git a/Engine/Extras/RoboMerge/v3/src/robo/roboargs.ts b/Engine/Extras/RoboMerge/v3/src/robo/roboargs.ts
--- a/Engine/Extras/RoboMerge/v3/src/robo/roboargs.ts
+++ b/Engine/Extras/RoboMerge/v3/src/robo/roboargs.ts
@@ -2,11 +2,19 @@
 
 import { Args } from '../common/args';
 
-let args: Args
-export function roboArgsInit(inArgs: Args) {
-	args = inArgs
-}
-
-export function getExclusiveLockOpenedsToRun() {
-	return args.exclusiveLockOpenedsToRun
-} 
\ No newline at end of file
+export class RoboArgs {
+
+	private static _args: Args
+	
+	static Init(args: Args) {
+		this._args = args
+	}
+
+	static getExclusiveLockOpenedsToRun() {
+		return this._args.exclusiveLockOpenedsToRun
+	}
+
+	static get args() {
+		return this._args
+	}
+}
\ No newline at end of file
