If MongoURI defined in the arguments settings will be saved to a Mongo DB rather than to disk

diff --git a/Engine/Extras/RoboMerge/v3/package.json b/Engine/Extras/RoboMerge/v3/package.json
--- a/Engine/Extras/RoboMerge/v3/package.json
+++ b/Engine/Extras/RoboMerge/v3/package.json
@@ -28,6 +28,7 @@
     "ldapauth-fork": "^4.0.2",
     "markdown-toc": "^1.2.0",
     "marked": "^4.0.12",
+    "mongoose": "^8.13.2",
     "mustache": "^3.2.1",
     "node-gzip": "^1.1.2",
     "nodemailer": "^6.7.0",
diff --git a/Engine/Extras/RoboMerge/v3/src/robo/graphbot.ts b/Engine/Extras/RoboMerge/v3/src/robo/graphbot.ts
--- a/Engine/Extras/RoboMerge/v3/src/robo/graphbot.ts
+++ b/Engine/Extras/RoboMerge/v3/src/robo/graphbot.ts
@@ -92,7 +92,7 @@ export class GraphBot implements BotEventHandler {
 			graphBot.botLogger.error(`Problem starting up bot ${botname}: ${error}`);
 		}
 
-		graphBot._settings = new Settings(botname, graphBot.botLogger, graphBot.p4)
+		graphBot._settings = await Settings.CreateAsync(botname, graphBot.botLogger, graphBot.p4)
 
 		graphBot.externalUrl = externalUrl
 
diff --git a/Engine/Extras/RoboMerge/v3/src/robo/models/bot.ts b/Engine/Extras/RoboMerge/v3/src/robo/models/bot.ts
new file mode 100644
--- /dev/null
+++ b/Engine/Extras/RoboMerge/v3/src/robo/models/bot.ts
@@ -0,0 +1,13 @@
+// Copyright Epic Games, Inc. All Rights Reserved.
+
+var mongoose = require('mongoose');
+
+/**
+ * Bot schema
+ */
+var botSchema = new mongoose.Schema({
+  botname: { type: String, index: true },
+  details: {  }
+});
+
+module.exports = mongoose.model('Bot', botSchema);
diff --git a/Engine/Extras/RoboMerge/v3/src/robo/robo.ts b/Engine/Extras/RoboMerge/v3/src/robo/robo.ts
--- a/Engine/Extras/RoboMerge/v3/src/robo/robo.ts
+++ b/Engine/Extras/RoboMerge/v3/src/robo/robo.ts
@@ -189,6 +189,11 @@ const COMMAND_LINE_ARGS: {[param: string]: Arg<any>} = {
 		env: 'ROBO_SLACK_DOMAIN',
 		dflt: 'https://slack.com'
 	},
+	mongoDB_URI: {
+		match: /^-mongoDB_URI=(.+)$/,
+		env: 'ROBO_MONGODB_URI',
+		dflt: ''
+	},
 	persistenceDir: {
 		match: /^-persistenceDir=(.+)$/,
 		env: 'ROBO_PERSISTENCE_DIR',
@@ -225,6 +230,18 @@ const vault = new Vault(args.vault)
 notificationsInit(args, vault)
 settingsInit(args)
 
+/**
+ * Connect to MongoDB.
+ */
+if (RoboArgs.useMongo()) {
+	var mongoose = require('mongoose');
+	mongoose.connect(args.mongoDB_URI);
+	  mongoose.connection.on('error', (err: any) => {
+		console.log(`MongoDB Connection Error: ${err}`);
+		process.exit(1);
+	  });
+}
+
 const env = args.epicEnv || (args.nodeEnv === "production" ? "prod" : "dev")
 const sentryEnv = `${env}-${args.epicDeployment}`
 if (env === 'prod' 
diff --git a/Engine/Extras/RoboMerge/v3/src/robo/roboargs.ts b/Engine/Extras/RoboMerge/v3/src/robo/roboargs.ts
--- a/Engine/Extras/RoboMerge/v3/src/robo/roboargs.ts
+++ b/Engine/Extras/RoboMerge/v3/src/robo/roboargs.ts
@@ -14,6 +14,10 @@ export class RoboArgs {
 		return this._args.exclusiveLockOpenedsToRun
 	}
 
+	static useMongo() {
+		return this._args.mongoDB_URI.length > 0
+	}
+
 	static get args() {
 		return this._args
 	}
diff --git a/Engine/Extras/RoboMerge/v3/src/robo/settings.ts b/Engine/Extras/RoboMerge/v3/src/robo/settings.ts
--- a/Engine/Extras/RoboMerge/v3/src/robo/settings.ts
+++ b/Engine/Extras/RoboMerge/v3/src/robo/settings.ts
@@ -5,7 +5,10 @@ import { Args } from '../common/args';
 import { ContextualLogger } from '../common/logger';
 import { VersionReader } from '../common/version';
 import { PerforceContext, Workspace } from '../common/perforce';
-import semver = require('semver')
+import { RoboArgs } from './roboargs';
+import semver = require('semver');
+
+const BotModel = require('./models/bot');
 
 /**
  * Descriptions of major versions in settings layout
@@ -94,59 +97,78 @@ const FilePurposeExtensionMap = {
 }
 
 function getPersistenceFilepath(botname: string) {
-	return args.persistenceDir + `/${botname.toLowerCase()}.settings`
+	return args.persistenceDir + `/${botname}.settings`
 }
 
 export class Settings {
-	filename: string
+	botname: string
 	enableSave: boolean
 	object: { [key: string]: any }
 	private	lastPersistBackupTime: Date
 	private readonly settingsLogger: ContextualLogger
-	private readonly p4workspace?: Workspace
+	private readonly p4backupWorkspace?: Workspace
 
-	constructor(botname: string, parentLogger: ContextualLogger, private p4: PerforceContext) {
+	private constructor(botname: string, parentLogger: ContextualLogger, private p4: PerforceContext) {
 		this.settingsLogger = parentLogger.createChild('Settings')
 
 		this.lastPersistBackupTime = new Date(Date.now())
-		this.filename = getPersistenceFilepath(botname)
+		this.botname = botname.toLowerCase() 
 		if (args.persistenceBackupFrequency > 0) {
-			this.p4workspace = {directory: fs.realpathSync(args.persistenceDir), name: args.persistenceBackupWorkspace}
+			this.p4backupWorkspace = {directory: fs.realpathSync(args.persistenceDir), name: args.persistenceBackupWorkspace}
 		}
 
 		// see if we should enable saves
 		this.enableSave = !process.env["NOSAVE"];
-		if (this.enableSave)
-			this.settingsLogger.verbose(`Reading settings from ${this.getFilename('PERSISTENCE')}`);
-		else
+		if (!this.enableSave) {
 			this.settingsLogger.info("Saving config has been disabled by NOSAVE environment variable");
+		}
+	}
+
+	static async CreateAsync(botname: string, parentLogger: ContextualLogger, p4: PerforceContext) {
+
+		const settings = new Settings(botname, parentLogger, p4)
+
+		if (RoboArgs.useMongo()) {
+			// first try to get the bot data from the DB
+			const botRecord = await BotModel.find({botname: botname})
+			if (botRecord.length > 1) {
+				throw new Error(`Found multiple records for ${botname}`) 
+			}
+
+			if (botRecord.length == 1) {
+				settings.object = botRecord[0].details
+			}
+		}
 
-		// load the object from disk
-		let filebits = readFileToString(this.getFilename('PERSISTENCE'));
 		let objVersion : semver.SemVer
-		if (filebits)
-			this.object = jsonlint.parse(filebits);
-		else {
-			// Create "empty" settings object, but include latest version so we don't needless enter migration code
-			this.object = {
-				version: VERSION.raw
-			};
-			objVersion = VERSION
+		if (!settings.object)
+		{
+			// load the object from disk
+			let filebits = readFileToString(settings.getFilename('PERSISTENCE'));
+			if (filebits)
+				settings.object = jsonlint.parse(filebits);
+			else {
+				// Create "empty" settings object, but include latest version so we don't needless enter migration code
+				settings.object = {
+					version: VERSION.raw
+				};
+				objVersion = VERSION
+			}
 		}
 
 		// Originally we did not version configuration.
 		// If we have no version, assume it needs all migrations
-		if (!this.object.version) {
-			this.settingsLogger.warn("No version found in settings data.")
+		if (!settings.object.version) {
+			settings.settingsLogger.warn("No version found in settings data.")
 			objVersion = new semver.SemVer('0.0.0')
 		} 
 		// Ensure we have a semantic version 
 		else {
-			const fileSemVer = semver.coerce(String(this.object.version))
+			const fileSemVer = semver.coerce(String(settings.object.version))
 			if (fileSemVer) {
 				objVersion = fileSemVer
 			} else {
-				throw new Error(`Found version field in settings file, but it does not appear to be a semantic version: "${this.object.version}"`)
+				throw new Error(`Found version field in settings file, but it does not appear to be a semantic version: "${settings.object.version}"`)
 			}
 		}
 
@@ -162,15 +184,17 @@ export class Settings {
 
 		// We're up to date!
 		if (previousObjVersion.compare(VERSION) !== 0) {
-			this.settingsLogger.info(`Changing settings version from '${previousObjVersion.raw}' to '${VERSION.raw}'`)
+			settings.settingsLogger.info(`Changing settings version from '${previousObjVersion.raw}' to '${VERSION.raw}'`)
 		}
 
 		/**
 		 * MIGRATION CODE SECTION FINISH
 		 */
 		
-		this.object.version = VERSION.raw
-		this._saveObject(true);
+		settings.object.version = VERSION.raw
+		settings._saveObject(true);
+
+		return settings
 	}
 
 	getContext(name: string) {
@@ -178,47 +202,52 @@ export class Settings {
 	}
 
 	private getFilename(purpose: FilePurpose) {
-		return this.filename + FilePurposeExtensionMap[purpose]
+		return getPersistenceFilepath(this.botname) + FilePurposeExtensionMap[purpose]
 	}
 
 	_saveObject(ensureInP4?: boolean) {
 		if (this.enableSave) {
 
-			const persistentFile = this.getFilename('PERSISTENCE')
-			const tempFile = this.getFilename('TEMP')
-
-			if (fs.existsSync(persistentFile)) {
-				fs.copyFileSync(persistentFile,this.getFilename('BACKUP'))
+			if (RoboArgs.useMongo()) {
+				BotModel.collection.updateOne({botname:this.botname},{$set:{botname:this.botname,details:this.object}},{upsert:true})
 			}
+			else {
+				const persistentFile = this.getFilename('PERSISTENCE')
+				const tempFile = this.getFilename('TEMP')
 
-			let filebits = JSON.stringify(this.object, null, '  ')
-			fs.writeFileSync(tempFile, filebits, "utf8")
-
-			fs.copyFileSync(tempFile, persistentFile)
-
-			if (this.p4workspace) {
-				if (ensureInP4) {
-					this.p4.fstat(this.p4workspace, persistentFile)
-					.then((fstat) => {
-						if (fstat.length == 0 || fstat[0].headAction == 'delete') {
-							this.p4.new_cl(this.p4workspace!, "Adding persistent backup", [])
-							.then((cl: number) => {
-								this.p4.add(this.p4workspace!, cl, persistentFile, "text+wS64")
-								.then(() => this.p4.submit(this.p4workspace!, cl))
-							})
-						}
-					})
+				if (fs.existsSync(persistentFile)) {
+					fs.copyFileSync(persistentFile,this.getFilename('BACKUP'))
 				}
-				else if ((Date.now() - this.lastPersistBackupTime.getTime()) / 60000 >= args.persistenceBackupFrequency) {
-					this.lastPersistBackupTime = new Date(Date.now())
-
-					this.p4.sync(this.p4workspace, persistentFile, {opts: ['-k']})
-					.then(() => this.p4.new_cl(this.p4workspace!, "Updating persistent backup"))
-					.then((cl: number) => {
-						this.p4.edit(this.p4workspace!, cl, persistentFile, ['-k'])
-						.then(() => this.p4.submit(this.p4workspace!, cl))
-					})
-					
+
+				let filebits = JSON.stringify(this.object, null, '  ')
+				fs.writeFileSync(tempFile, filebits, "utf8")
+
+				fs.copyFileSync(tempFile, persistentFile)
+
+				if (this.p4backupWorkspace) {
+					if (ensureInP4) {
+						this.p4.fstat(this.p4backupWorkspace, persistentFile)
+						.then((fstat) => {
+							if (fstat.length == 0 || fstat[0].headAction == 'delete') {
+								this.p4.new_cl(this.p4backupWorkspace!, "Adding persistent backup", [])
+								.then((cl: number) => {
+									this.p4.add(this.p4backupWorkspace!, cl, persistentFile, "text+wS64")
+									.then(() => this.p4.submit(this.p4backupWorkspace!, cl))
+								})
+							}
+						})
+					}
+					else if ((Date.now() - this.lastPersistBackupTime.getTime()) / 60000 >= args.persistenceBackupFrequency) {
+						this.lastPersistBackupTime = new Date(Date.now())
+
+						this.p4.sync(this.p4backupWorkspace, persistentFile, {opts: ['-k']})
+						.then(() => this.p4.new_cl(this.p4backupWorkspace!, "Updating persistent backup"))
+						.then((cl: number) => {
+							this.p4.edit(this.p4backupWorkspace!, cl, persistentFile, ['-k'])
+							.then(() => this.p4.submit(this.p4backupWorkspace!, cl))
+						})
+						
+					}
 				}
 			}
 		}
