Persist slack messages in the DB in a separate table

diff --git a/Engine/Extras/RoboMerge/v3/src/robo/bot-interfaces.ts b/Engine/Extras/RoboMerge/v3/src/robo/bot-interfaces.ts
--- a/Engine/Extras/RoboMerge/v3/src/robo/bot-interfaces.ts
+++ b/Engine/Extras/RoboMerge/v3/src/robo/bot-interfaces.ts
@@ -72,7 +72,7 @@ export interface NodeBotInterface extends NodeBotIPC {
 	getImmediateEdge(branchName: BranchArg): EdgeBotInterface | null
 
 	getNumConflicts(): number
-	applyStatus(status: BranchStatus): void
+	applyStatus(status: BranchStatus): Promise<void>
 
 }
 
diff --git a/Engine/Extras/RoboMerge/v3/src/robo/conflicts.ts b/Engine/Extras/RoboMerge/v3/src/robo/conflicts.ts
--- a/Engine/Extras/RoboMerge/v3/src/robo/conflicts.ts
+++ b/Engine/Extras/RoboMerge/v3/src/robo/conflicts.ts
@@ -408,7 +408,7 @@ export class Conflicts {
 		return null
 	}
 
-	applyStatus(conflictsForStatus: Partial<ConflictStatusFields>[], slackMessages?: SlackMessages) {
+	async applyStatus(conflictsForStatus: Partial<ConflictStatusFields>[], slackMessages?: SlackMessages) {
 		for (const conflict of this.conflicts) {
 			const conflictObj: Partial<ConflictStatusFields> = {}
 
@@ -428,8 +428,8 @@ export class Conflicts {
 			conflictObj.owner = conflict.owner
 
 			if (slackMessages) {
-				const findResult = slackMessages.findAll(conflict.cl, conflict.targetBranchName || conflict.kind)
-				for (const result of findResult.messages) {
+				const messages = await slackMessages.findAll(conflict.cl, conflict.targetBranchName || conflict.kind)
+				for (const result of messages) {
 					if (result.messageOpts.channel.startsWith('C') && result.permalink) {
 						if (!conflictObj.slackLinks) {
 							conflictObj.slackLinks = []
diff --git a/Engine/Extras/RoboMerge/v3/src/robo/graphbot.ts b/Engine/Extras/RoboMerge/v3/src/robo/graphbot.ts
--- a/Engine/Extras/RoboMerge/v3/src/robo/graphbot.ts
+++ b/Engine/Extras/RoboMerge/v3/src/robo/graphbot.ts
@@ -13,8 +13,8 @@ import { BranchGraph } from './branchgraph';
 import { PersistentConflict } from './conflict-interfaces';
 import { BotEventHandler, BotEventTriggers } from './events';
 import { NodeBot } from './nodebot';
-import { bindBotNotifications, BotNotifications, NOTIFICATIONS_PERSISTENCE_KEY } from './notifications';
-import { postMessageToChannel, postToRobomergeAlerts, SlackMessages } from './notifications';
+import { bindBotNotifications, BotNotifications, postMessageToChannel, postToRobomergeAlerts, SlackMessages } from './notifications';
+import { PersistedSlackMessages } from './persistedmessages';
 import { roboAnalytics } from './roboanalytics';
 import { BlockageNodeOpUrls, OperationUrlHelper } from './roboserver';
 import { Settings } from './settings';
@@ -129,7 +129,7 @@ export class GraphBot implements BotEventHandler {
 			}
 		}
 
-		this.slackMessages = bindBotNotifications(this.eventTriggers, slackChannelOverrides, this.settings.getContext(NOTIFICATIONS_PERSISTENCE_KEY), blockageUrlGenerator, this.externalUrl, this.botLogger)
+		this.slackMessages = bindBotNotifications(this.eventTriggers, slackChannelOverrides, PersistedSlackMessages.get(this.settings, this.botLogger), blockageUrlGenerator, this.externalUrl, this.botLogger)
 		bindBadgeHandler(this.eventTriggers, this.branchGraph, this.externalUrl, this.botLogger)
 
 		let hasConflicts = false
diff --git a/Engine/Extras/RoboMerge/v3/src/robo/ipc.ts b/Engine/Extras/RoboMerge/v3/src/robo/ipc.ts
--- a/Engine/Extras/RoboMerge/v3/src/robo/ipc.ts
+++ b/Engine/Extras/RoboMerge/v3/src/robo/ipc.ts
@@ -86,21 +86,22 @@ export class IPC {
 		}
 	}
 
-	private getBranches() {
+	private async getBranches() {
 		const status = new Status(START_TIME, this.robo.VERSION_STRING, this.ipcLogger)
 		for (let graphBot of this.robo.graphBots.values()) {
 			graphBot.applyStatus(status)
-			for (let branch of graphBot.branchGraph.branches) {
+			await Promise.all(graphBot.branchGraph.branches.map(branch => {
 				if (branch.isMonitored) {
-					status.addBranch(branch)
+					return status.addBranch(branch)
 				}
-			}
+				return undefined
+			}))
 		}
 
 		return { ...OPERATION_SUCCESS, data: status }
 	}
 
-	private getBranch(botname: string, branchname: string) {
+	private async getBranch(botname: string, branchname: string) {
 		const status = new Status(START_TIME, this.robo.VERSION_STRING, this.ipcLogger)
 		botname = botname.toUpperCase()
 		branchname = branchname.toUpperCase()
@@ -110,7 +111,7 @@ export class IPC {
 		if (graphBot) {
 			const branch = graphBot.branchGraph.getBranch(branchname)
 			if (branch) {
-				status.addBranch(branch)
+				await status.addBranch(branch)
 			}
 		}
 
diff --git a/Engine/Extras/RoboMerge/v3/src/robo/models/slackmessages.ts b/Engine/Extras/RoboMerge/v3/src/robo/models/slackmessages.ts
new file mode 100644
--- /dev/null
+++ b/Engine/Extras/RoboMerge/v3/src/robo/models/slackmessages.ts
@@ -0,0 +1,19 @@
+// Copyright Epic Games, Inc. All Rights Reserved.
+
+var mongoose = require('mongoose');
+
+/**
+ * Bot schema
+ */
+var slackMessagesSchema = new mongoose.Schema({
+  botname: { type: String, index: true },
+  cl: { type: Number, index: true },
+  branchName: { type: String, index: true},
+  channel: { type: String, index: true },
+  timestamp: { type: String },
+  timestamp_forSunset: { type: Number },
+  permalink: { type: String },
+  message: {  }
+});
+
+module.exports = mongoose.model('SlackMessages', slackMessagesSchema);
diff --git a/Engine/Extras/RoboMerge/v3/src/robo/nodebot.ts b/Engine/Extras/RoboMerge/v3/src/robo/nodebot.ts
--- a/Engine/Extras/RoboMerge/v3/src/robo/nodebot.ts
+++ b/Engine/Extras/RoboMerge/v3/src/robo/nodebot.ts
@@ -1297,7 +1297,7 @@ export class NodeBot extends PerforceStatefulBot implements NodeBotInterface {
 		return this.edges.size
 	}
 
-	applyStatus(status: BranchStatus) {
+	async applyStatus(status: BranchStatus) {
 		status.display_name = this.displayName
 		status.last_cl = this.lastCl
 		status.swarmURL = this.p4.swarmURL
@@ -1327,7 +1327,7 @@ export class NodeBot extends PerforceStatefulBot implements NodeBotInterface {
 		}
 
 		status.conflicts = [] as any[]
-		this.conflicts.applyStatus(status.conflicts, this.slackMessages)
+		await this.conflicts.applyStatus(status.conflicts, this.slackMessages)
 
 		status.tick_count = this.tickCount
 
diff --git a/Engine/Extras/RoboMerge/v3/src/robo/notifications.ts b/Engine/Extras/RoboMerge/v3/src/robo/notifications.ts
--- a/Engine/Extras/RoboMerge/v3/src/robo/notifications.ts
+++ b/Engine/Extras/RoboMerge/v3/src/robo/notifications.ts
@@ -10,18 +10,15 @@ import { Blockage, Branch, BranchArg,  ExclusiveLockInfo, ForcedCl, MergeAction,
 import { PersistentConflict, Resolution } from './conflict-interfaces';
 import { BotEventHandler, BotEvents } from './events';
 import { NodeBot } from './nodebot';
+import { PersistedSlackMessage, PersistedSlackMessages } from './persistedmessages';
 import { BlockageNodeOpUrls } from './roboserver';
-import { Context } from './settings';
 import { Slack, SlackAttachment, SlackFile, SlackLinkButtonAction, SlackLinkButtonsAttachment, SlackMessageField, SlackMessage, SlackMessageStyles } from './slack';
 import { Vault } from './vault'
 import { WebServer } from '../common/webserver'
 import { DummySlackApp } from '../common/dummyslackserver'
 
-export const NOTIFICATIONS_PERSISTENCE_KEY = 'notifications'
-export const SLACK_MESSAGES_PERSISTENCE_KEY = 'slackMessages'
 const CHANGELIST_FIELD_TITLE = 'Change'
 const ACKNOWLEDGED_FIELD_TITLE = 'Acknowledged'
-const SUNSET_PERSISTED_NOTIFICATIONS_DAYS = 30
 const EPIC_TIME_OPTIONS: DateTimeFormatOptions = {timeZone: 'EST5EDT', timeZoneName: 'short'}
 
 const KNOWN_BOT_NAMES = ['buildmachine', 'robomerge'];
@@ -56,7 +53,6 @@ export function isUserAKnownBot(user: string) {
 	return KNOWN_BOT_NAMES.indexOf(user) !== -1 || KNOWN_BOT_EMAILS.indexOf(user) !== -1
 }
 
-
 export async function postToRobomergeAlerts(message: string) {
 	// 'C044ZUSS71S' is #dt-robomerge-alert-ext
 	if (!args.devMode) {
@@ -97,11 +93,6 @@ export async function getPingTarget(userEmail: string|null, owner: string, slack
 		: `@${owner}`
 }
 
-function generatePersistedSlackMessageKey(sourceCl: number, targetBranchArg: BranchArg, channel: string) {
-	const branchName = resolveBranchArg(targetBranchArg, true)
-	return `${sourceCl}:${branchName}:${channel}`
-}
-
 function formatDuration(durationSeconds: number) {
 	const underSixHours = durationSeconds < 6 * 3600
 	const durationBits: string[] = []
@@ -203,57 +194,38 @@ function formatResolution(info: PersistentConflict) {
 	return message
 }
 
-//////////
-
-type SlackConflictMessage = {
-	timestamp: string
-	permalink: string
-	messageOpts: SlackMessage
-}
-
-type PersistedSlackMessages = {[key: string]: SlackConflictMessage}
-
 /** Wrapper around Slack, keeping track of messages keyed on target branch and source CL */
 export class SlackMessages {
 	private readonly smLogger: ContextualLogger
-	constructor(private slack: Slack, private persistence: Context, parentLogger: ContextualLogger) {
+	constructor(private slack: Slack, private persistedMessages: PersistedSlackMessages, parentLogger: ContextualLogger) {
 		this.smLogger = parentLogger.createChild('SlackMsgs')
-		const slackMessages: PersistedSlackMessages | null = this.persistence.get(SLACK_MESSAGES_PERSISTENCE_KEY)
-
-		if (slackMessages) {
-			this.doHouseKeeping(slackMessages)
-		}
-		else {
-			this.persistence.set(SLACK_MESSAGES_PERSISTENCE_KEY, {})
-		}
 	}
 
 	async postOrUpdate(cl: number, branchArg: BranchArg, message: SlackMessage, persistMessage = true) {
-		const findResult = this.find(cl, branchArg, message.channel)
+		const findResult = await this.find(cl, branchArg, message.channel)
 
 		// If we find a message, simply update the contents
-		if (findResult.messageRecord) {
+		if (findResult) {
 			// keep ack field if present and not in new message
 
-			if (findResult.messageRecord.messageOpts.fields && (
+			if (findResult.messageOpts.fields && (
 				!message.fields || !message.fields.some(field =>
 				field.title === ACKNOWLEDGED_FIELD_TITLE))) {
 
-				for (const field of findResult.messageRecord.messageOpts.fields) {
+				for (const field of findResult.messageOpts.fields) {
 					if (field.title === ACKNOWLEDGED_FIELD_TITLE) {
 						message.fields = [...(message.fields || []), field]
 						break
 					}
 				}
 			}
-			findResult.messageRecord.messageOpts = message
-			await this.update(findResult.messageRecord, findResult.persistedMessages)
+			findResult.messageOpts = message
+			await this.update(findResult)
 
-			return findResult.messageRecord.permalink
+			return findResult.permalink
 		}
 		// Otherwise, we will need to create a new one
 		else {
-
 			message.target = resolveBranchArg(branchArg, true)
 			message.cl = cl
 
@@ -270,8 +242,7 @@ export class SlackMessages {
 
 			// Used for messages we don't care to keep, currently the /api/test/directmessage endpoint
 			if (persistMessage) {
-				findResult.persistedMessages[findResult.conflictKey] = {timestamp, permalink, messageOpts: message}
-				this.persistence.set(SLACK_MESSAGES_PERSISTENCE_KEY, findResult.persistedMessages)
+				this.persistedMessages.add(message, timestamp, permalink)
 			}
 
 			return permalink
@@ -279,11 +250,11 @@ export class SlackMessages {
 	}
 
 	async postReply(cl: number, branchArg: BranchArg, message: SlackMessage) {
-		const findResult = this.find(cl, branchArg, message.channel)
+		const findResult = await this.find(cl, branchArg, message.channel)
 
-		if (findResult.messageRecord) {
+		if (findResult) {
 			try {
-				await this.slack.reply(findResult.messageRecord.timestamp, message)
+				await this.slack.reply(findResult.timestamp, message)
 			}
 			catch (err) {
 				console.error('Error updating message in Slack! ' + err.toString())
@@ -296,11 +267,11 @@ export class SlackMessages {
 	}
 
 	async postFile(cl: number, branchArg: BranchArg, file: SlackFile) {
-		const findResult = this.find(cl, branchArg, file.channels)
+		const findResult = await this.find(cl, branchArg, file.channels)
 
-		if (findResult.messageRecord) {
+		if (findResult) {
 			try {
-				file.thread_ts = findResult.messageRecord.timestamp
+				file.thread_ts = findResult.timestamp
 				await this.slack.uploadFile(file)
 			}
 			catch (err) {
@@ -348,7 +319,7 @@ export class SlackMessages {
 		if (channelId) {
 			dm.channel = channelId
 			// Add the channel/conversation ID to the messageOpts and proceed normally.
-			this.smLogger.info(`Creating direct message for ${emailAddress} (key: ${generatePersistedSlackMessageKey(cl, branchArg, channelId)})`)
+			this.smLogger.info(`Creating direct message for ${emailAddress} (key: ${PersistedSlackMessages.generateKey(cl, branchArg, channelId)})`)
 			this.postOrUpdate(cl, branchArg, dm, persistMessage)
 		}
 	}
@@ -363,41 +334,24 @@ export class SlackMessages {
 		this.postFile(cl, branchArg, file)
 	}
 
-	// 
 	findAll(cl: number, branchArg: BranchArg) {
-		// Get key to search our list of all messages
-		const persistedMessages = this.persistence.get(SLACK_MESSAGES_PERSISTENCE_KEY) as PersistedSlackMessages
-		const conflictKey = generatePersistedSlackMessageKey(cl, branchArg, "") // pass empty string as channel name
-
-		// Filter messages
-		const messages: SlackConflictMessage[] = []
-
-		for (const key in persistedMessages) {
-			if (key.startsWith(conflictKey)) {
-				messages.push(persistedMessages[key])
-			}
-		}
-
-		return {messages, persistedMessages}
+		return this.persistedMessages.findAll(cl, branchArg)
 	}
 
 	find(cl: number, branchArg: BranchArg, channel: string) {
-		// Get key to search our list of all messages
-		const persistedMessages = this.persistence.get(SLACK_MESSAGES_PERSISTENCE_KEY) as PersistedSlackMessages
-		const conflictKey = generatePersistedSlackMessageKey(cl, branchArg, channel)
-		return {messageRecord: persistedMessages[conflictKey], conflictKey, persistedMessages}
+		return this.persistedMessages.find(cl, branchArg, channel)
 	}
 
-	async update(messageRecord: SlackConflictMessage, persistedMessages: PersistedSlackMessages) {
+	async update(persistedMessage: PersistedSlackMessage) {
 		try {
-			await this.slack.update(messageRecord.timestamp, messageRecord.messageOpts)
+			await this.slack.update(persistedMessage.timestamp, persistedMessage.messageOpts)
 		}
 		catch (err) {
 			console.error('Error updating message in Slack! ' + err.toString())
 			return
 		}
 
-		this.persistence.set(SLACK_MESSAGES_PERSISTENCE_KEY, persistedMessages)
+		persistedMessage.update()
 	}
 
 	/** Deliberately ugly function name to avoid accidentally posting duplicate messages about conflicts! */
@@ -459,27 +413,6 @@ export class SlackMessages {
 		
 		return this.slack.lookupUserIdByEmail(emailAddress.toLowerCase())
 	}
-
-	private doHouseKeeping(messages: PersistedSlackMessages) {
-		const nowTicks = Date.now() / 1000.
-
-		const keys = Object.keys(messages)
-		for (const key of keys) {
-			const message = messages[key]
-			const messageAgeDays = (nowTicks - parseFloat(message.timestamp))/24/60/60
-
-			// note: serialisation was briefly wrong (11/11/2019), leaving some message records without a message
-			if (!message.messageOpts || messageAgeDays > SUNSET_PERSISTED_NOTIFICATIONS_DAYS) {
-				this.smLogger.info(`Sunsetting Slack Message "${key}"`)
-				delete messages[key]
-			}
-		}
-
-		// if we removed any messages, persist
-		if (Object.keys(messages).length < keys.length) {
-			this.persistence.set(SLACK_MESSAGES_PERSISTENCE_KEY, messages)
-		}
-	}
 }
 
 export function makeClLink(cl: number, alias?: string) {
@@ -747,11 +680,14 @@ export class BotNotifications implements BotEventHandler {
 
 			const messageText = formatResolution(info)
 			const targetKey = info.targetBranchName || info.kind
-			if (!this.updateMessagesAfterUnblock(info.cl, targetKey, '', messageStyle, messageText, newClDesc)) {
-				this.botNotificationsLogger.warn(`Conflict message not found to update (${info.blockedBranchName} -> ${targetKey} CL#${info.sourceCl})`)
-				const message = this.makeSlackChannelMessage('', messageText, messageStyle, makeClLink(info.cl), info.blockedBranchName, info.targetBranchName, info.author)
-				this.slackMessages.postOrUpdate(info.cl, targetKey, message)
-			}
+			this.updateMessagesAfterUnblock(info.cl, targetKey, '', messageStyle, messageText, newClDesc)
+			.then(success => {
+				if (!success) {
+					this.botNotificationsLogger.warn(`Conflict message not found to update (${info.blockedBranchName} -> ${targetKey} CL#${info.sourceCl})`)
+					const message = this.makeSlackChannelMessage('', messageText, messageStyle, makeClLink(info.cl), info.blockedBranchName, info.targetBranchName, info.author)
+					this.slackMessages!.postOrUpdate(info.cl, targetKey, message)
+				}
+			})
 		}
 	}
 
@@ -983,16 +919,16 @@ export class BotNotifications implements BotEventHandler {
 	
 
 	/** Pre-condition: this.slackMessages must be valid */
-	private updateMessagesAfterUnblock(cl: number, targetBranch: BranchArg, newTitle: string,
+	private async updateMessagesAfterUnblock(cl: number, targetBranch: BranchArg, newTitle: string,
 										newStyle: SlackMessageStyles, newText: string, newClDesc?: string) {
 		// Find all messages relating to CL and branch
-		const findResult = this.slackMessages!.findAll(cl, targetBranch)
+		const messages = await this.slackMessages!.findAll(cl, targetBranch)
 
-		if (findResult.messages.length == 0) {
+		if (messages.length == 0) {
 			return false
 		}
 
-		for (const messageRecord of findResult.messages) {
+		for (const messageRecord of messages) {
 			const message = messageRecord.messageOpts
 			if (newTitle) {
 				message.title = newTitle
@@ -1043,7 +979,7 @@ export class BotNotifications implements BotEventHandler {
 				delete message.footer
 			}
 
-			this.slackMessages!.update(messageRecord, findResult.persistedMessages)
+			this.slackMessages!.update(messageRecord)
 		}
 		
 		return true
@@ -1052,50 +988,52 @@ export class BotNotifications implements BotEventHandler {
 	/** Pre-condition: this.slackMessages must be valid */
 	private tryAddFieldToChannelMessages(cl: number, targetBranch: BranchArg, newField: SlackMessageField) {
 		// Find all messages relating to CL and branch
-		const findResult = this.slackMessages!.findAll(cl, targetBranch)
+		this.slackMessages!.findAll(cl, targetBranch)
+		.then(messages => {
+			for (const messageRecord of messages) {
+				const message = messageRecord.messageOpts
+				if (message.attachments) {
+					// skip DMs
+					continue
+				}
 
-		for (const messageRecord of findResult.messages) {
-			const message = messageRecord.messageOpts
-			if (message.attachments) {
-				// skip DMs
-				continue
-			}
+				if (message.fields && message.fields.find(field => field.title === newField.title)) {
+					// do not add same field twice (shouldn't happen, but hey)
+					continue
+				}
 
-			if (message.fields && message.fields.find(field => field.title === newField.title)) {
-				// do not add same field twice (shouldn't happen, but hey)
-				continue
+				message.fields = [...(message.fields || []), newField]
+				this.slackMessages!.update(messageRecord)
 			}
-
-			message.fields = [...(message.fields || []), newField]
-			this.slackMessages!.update(messageRecord, findResult.persistedMessages)
-		}
+		})
 	}
 
 	/** Pre-condition: this.slackMessages must be valid */
 	private tryRemoveFieldFromChannelMessages(cl: number, targetBranch: BranchArg, fieldTitle: string) {
 		// Find all messages relating to CL and branch
-		const findResult = this.slackMessages!.findAll(cl, targetBranch)
-
-		for (const messageRecord of findResult.messages) {
-			const message = messageRecord.messageOpts
-			if (message.attachments || !message.fields) {
-				// skip DMs (expecting there to be fields usually, but skipping if not)
-				continue
-			}
+		this.slackMessages!.findAll(cl, targetBranch)
+		.then(messages => {
+			for (const messageRecord of messages) {
+				const message = messageRecord.messageOpts
+				if (message.attachments || !message.fields) {
+					// skip DMs (expecting there to be fields usually, but skipping if not)
+					continue
+				}
 
-			const replacementFields = message.fields.filter(field => field.title !== fieldTitle)
-			if (message.fields.length > replacementFields.length) {
-				message.fields = replacementFields
-				this.slackMessages!.update(messageRecord, findResult.persistedMessages)
+				const replacementFields = message.fields.filter(field => field.title !== fieldTitle)
+				if (message.fields.length > replacementFields.length) {
+					message.fields = replacementFields
+					this.slackMessages!.update(messageRecord)
+				}
 			}
-		}
+		})
 	}
 
 	private readonly slackMessages?: SlackMessages
 	private readonly additionalBlockChannelIds = new Map<string, [string, boolean]>()
 }
 
-export function bindBotNotifications(events: BotEvents, slackChannelOverrides: [Branch, Branch, string, boolean][], persistence: Context, blockageUrlGenerator: NodeOpUrlGenerator, 
+export function bindBotNotifications(events: BotEvents, slackChannelOverrides: [Branch, Branch, string, boolean][], persistence: PersistedSlackMessages, blockageUrlGenerator: NodeOpUrlGenerator, 
 	externalUrl: string, logger: ContextualLogger) {
 
 	let slackMessages
diff --git a/Engine/Extras/RoboMerge/v3/src/robo/persistedmessages.ts b/Engine/Extras/RoboMerge/v3/src/robo/persistedmessages.ts
new file mode 100644
--- /dev/null
+++ b/Engine/Extras/RoboMerge/v3/src/robo/persistedmessages.ts
@@ -0,0 +1,297 @@
+// Copyright Epic Games, Inc. All Rights Reserved.
+
+import { BranchArg, resolveBranchArg } from './branch-interfaces';
+import { ContextualLogger } from '../common/logger';
+import { RoboArgs } from './roboargs';
+import { Context, Settings } from './settings';
+import { SlackMessage } from './slack';
+
+const SlackMessagesModel = require('./models/slackmessages');
+
+const NOTIFICATIONS_PERSISTENCE_KEY = 'notifications'
+const SLACK_MESSAGES_PERSISTENCE_KEY = 'slackMessages'
+const SUNSET_PERSISTED_NOTIFICATIONS_DAYS = 30
+
+type PersistedMessagePayload = {
+	timestamp: string
+	permalink: string
+	messageOpts: SlackMessage
+}
+
+type PersistedMessagePersistence = {[key: string]: PersistedMessagePayload}
+
+interface PersistedMessageKeyComponents {
+	cl: number
+	branchName: string
+	channel: string
+}
+
+class PersistedMessageKey {
+
+	private _key: string | PersistedMessageKeyComponents
+
+	constructor(key: string)
+	constructor(cl: number, branchName: string, channel: string)
+	constructor(arg1: string|number, branchName?: string, channel?: string) {
+		if (typeof arg1 === "string") {
+			this._key = arg1
+		}
+		else {
+			this._key = {cl: arg1, branchName: branchName!, channel: channel!}
+		}
+	}
+
+	toString() {
+		if (typeof this._key === "string") {
+			return this._key
+		}
+		return PersistedSlackMessages.generateKey(this._key.cl, this._key.branchName, this._key.channel)
+	}
+	
+	toComponents() {
+		if (typeof this._key === "string") {
+			const components = this._key.split(':')
+			if (components.length != 3) {
+				throw new Error(`Unexpected number of components in ${this._key}`)
+			}
+			return {
+				cl: parseInt(components[0]),
+				branchName: components[1],
+				channel: components[2]
+			}
+		}
+		return this._key
+	}	
+}
+
+export class PersistedSlackMessage implements PersistedMessagePayload {
+	timestamp: string
+	permalink: string
+	messageOpts: SlackMessage
+
+	constructor(payload: PersistedMessagePayload, private _key: PersistedMessageKey, private persistedMessages: PersistedSlackMessages) {
+		this.timestamp = payload.timestamp
+		this.permalink = payload.permalink
+		this.messageOpts = payload.messageOpts
+	}
+
+	get key() {
+		return this._key
+	}
+
+	update() {
+		this.persistedMessages.update(this)
+	}
+}
+
+export abstract class PersistedSlackMessages {
+
+	static get(botSettings: Settings, logger: ContextualLogger): PersistedSlackMessages {
+		if (RoboArgs.useMongo()) {
+			return new DBPersistedSlackMessages(botSettings.botname)
+		}
+		return new FilePersistedSlackMessages(botSettings, logger)
+	}
+
+    static generateKey(sourceCl: number, targetBranchArg: BranchArg, channel: string) {
+        const branchName = resolveBranchArg(targetBranchArg, true)
+        return `${sourceCl}:${branchName}:${channel}`
+    }
+
+    abstract add(messageOpts: SlackMessage, timestamp: string, permalink: string): void
+	abstract find(cl: number, branchArg: BranchArg, channel: string): Promise<PersistedSlackMessage | undefined>
+	abstract findAll(cl: number, branchArg: BranchArg): Promise<PersistedSlackMessage[]>
+    abstract update(message: PersistedSlackMessage): void
+}
+
+class DBPersistedSlackMessages extends PersistedSlackMessages {
+
+	private messageCache: {[key: string]: PersistedSlackMessage} = {}
+
+	constructor(private botname: string) {
+		super()		
+	}
+
+	private write(message: PersistedSlackMessage) {
+		const index = {
+			botname: this.botname,
+			...message.key.toComponents()
+		}
+		return SlackMessagesModel.collection.updateOne(
+			index,
+			{$set:{
+				...index, 
+				timestamp: message.timestamp, 
+				permalink: message.permalink, 
+				message: message.messageOpts, 
+				timestamp_forSunset: parseFloat(message.timestamp)
+			}},
+			{upsert:true}
+		)
+	}
+
+    add(message: SlackMessage, timestamp: string, permalink: string) {
+		const key: PersistedMessageKey = new PersistedMessageKey(message.cl!, message.target!, message.channel)
+		const persistedMessage = new PersistedSlackMessage({timestamp, permalink, messageOpts: message}, key, this)
+		this.messageCache[key.toString()] = persistedMessage
+		this.write(persistedMessage)
+    }
+
+	async find(cl: number, branchArg: BranchArg, channel: string) {
+		const key = new PersistedMessageKey(cl, resolveBranchArg(branchArg, true), channel)
+		const cacheKey = key.toString()
+		if (this.messageCache[cacheKey]) {
+			return this.messageCache[cacheKey]
+		}
+		const index = {
+			botname: this.botname,
+			...key.toComponents()
+		}
+		const messageRecord = await SlackMessagesModel.find(index)
+		if (messageRecord.length == 1) {
+			const persistedMessage = new PersistedSlackMessage(messageRecord.message, key, this)
+			this.messageCache[cacheKey] = persistedMessage
+			return persistedMessage
+		}
+		else if (messageRecord.length > 1) {
+			throw new Error(`Found multiple records for ${this.botname} ${PersistedSlackMessages.generateKey(index.cl, index.branchName, index.channel)}`)
+		}
+
+		return undefined
+	}
+
+	async findAll(cl: number, branchArg: BranchArg): Promise<PersistedSlackMessage[]> {
+		const index = {
+			botname: this.botname, 
+			branchName: resolveBranchArg(branchArg, true),
+			cl
+		}
+		const messageRecords = await SlackMessagesModel.find(index)
+		const messages: PersistedSlackMessage[] = []
+
+		for (const messageRecord of messageRecords) {
+			const key = new PersistedMessageKey(cl, index.branchName, messageRecord.channel)
+			const cacheKey = key.toString()
+			if (!this.messageCache[cacheKey]) {
+				const persistedMessage = new PersistedSlackMessage(messageRecord.message, key, this)
+				this.messageCache[cacheKey] = persistedMessage
+			}
+			messages.push(this.messageCache[cacheKey])
+		}
+
+		return messages
+	}
+
+	update(message: PersistedSlackMessage) {
+		this.write(message)
+	}
+}
+
+export async function doDBPersistedMessagesHousekeeping() {
+	const expiredTime = Date.now() - (SUNSET_PERSISTED_NOTIFICATIONS_DAYS * 60 * 60 * 24 * 1000)
+	return SlackMessagesModel.collection.deleteMany({timestamp_forSunset: { $lt: expiredTime }})
+}
+
+export async function migrateMessagesToDB(botSettings: Settings) {
+	const filePersistedMessages = botSettings.getContext(NOTIFICATIONS_PERSISTENCE_KEY)?.get(SLACK_MESSAGES_PERSISTENCE_KEY)
+	if (filePersistedMessages) {
+		const dbPersistedMessages = new DBPersistedSlackMessages(botSettings.botname)
+		await Promise.all(Object.keys(filePersistedMessages).map(key => {
+			const message = filePersistedMessages[key]
+			// Some persisted messages don't have the cl and channel populated, so do so now from the key
+			const index = new PersistedMessageKey(key).toComponents()
+			message.messageOpts.cl = index.cl
+			message.messageOpts.target = index.branchName
+			message.messageOpts.channel = index.channel
+			return dbPersistedMessages.add(message.messageOpts, message.timestamp, message.permalink)
+		}))
+	}
+	delete botSettings.object[NOTIFICATIONS_PERSISTENCE_KEY]
+}
+
+class FilePersistedSlackMessages extends PersistedSlackMessages {
+
+	private persistence: Context
+
+	constructor(botSettings: Settings, logger: ContextualLogger) {
+		super()
+
+		this.persistence = botSettings.getContext(NOTIFICATIONS_PERSISTENCE_KEY)
+
+		if (this.persistence.get(SLACK_MESSAGES_PERSISTENCE_KEY)) {
+			this.doHouseKeeping(logger)
+		}
+		else {
+			this.persistence.set(SLACK_MESSAGES_PERSISTENCE_KEY, {})
+		}
+	}
+
+	private write(key: string, messageOpts: SlackMessage, timestamp: string, permalink: string) {
+		const persistedMessages: PersistedMessagePersistence = this.persistence.get(SLACK_MESSAGES_PERSISTENCE_KEY)
+        persistedMessages[key] = {timestamp, permalink, messageOpts}
+        this.persistence.set(SLACK_MESSAGES_PERSISTENCE_KEY, persistedMessages)
+	}
+
+    add(messageOpts: SlackMessage, timestamp: string, permalink: string) {
+		const conflictKey = PersistedSlackMessages.generateKey(messageOpts.cl!, messageOpts.target!, messageOpts.channel)
+		this.write(conflictKey, messageOpts, timestamp, permalink)
+    }
+
+	async find(cl: number, branchArg: BranchArg, channel: string) {
+		// Get key to search our list of all messages
+		const pms: PersistedMessagePersistence = this.persistence.get(SLACK_MESSAGES_PERSISTENCE_KEY)
+		const conflictKey = PersistedSlackMessages.generateKey(cl, branchArg, channel)
+		const persistedMessage = pms[conflictKey]
+		if (persistedMessage) {
+			return new PersistedSlackMessage(persistedMessage, new PersistedMessageKey(conflictKey), this)
+		}
+		return undefined
+	}
+
+	// 
+	async findAll(cl: number, branchArg: BranchArg) {
+		// Get key to search our list of all messages
+		const persistedMessages: PersistedMessagePersistence = this.persistence.get(SLACK_MESSAGES_PERSISTENCE_KEY)
+		const conflictKey = PersistedSlackMessages.generateKey(cl, branchArg, "") // pass empty string as channel name
+
+		// Filter messages
+		const messages: PersistedSlackMessage[] = []
+
+		for (const key in persistedMessages) {
+			if (key.startsWith(conflictKey)) {
+				messages.push(new PersistedSlackMessage(persistedMessages[key], new PersistedMessageKey(key), this))
+			}
+		}
+
+		return messages
+	}
+
+	update(message: PersistedSlackMessage) {
+		this.write(message.key.toString(), message.messageOpts, message.timestamp, message.permalink)
+	}	
+
+	private doHouseKeeping(logger: ContextualLogger) {
+		const expiredTime = Date.now() - (SUNSET_PERSISTED_NOTIFICATIONS_DAYS * 60 * 60 * 24 * 1000)
+
+		const persistedMessages: PersistedMessagePersistence = this.persistence.get(SLACK_MESSAGES_PERSISTENCE_KEY)
+
+		let persist = false
+		const keys = Object.keys(persistedMessages)
+		for (const key of keys) {
+			const message = persistedMessages[key]
+			const messageTime = parseFloat(message.timestamp)
+
+			// note: serialisation was briefly wrong (11/11/2019), leaving some message records without a message
+			if (!message.messageOpts || messageTime < expiredTime) {
+				logger.info(`Sunsetting Slack Message "${key}"`)
+				delete persistedMessages[key]
+				persist = true
+			}
+		}
+
+		// if we removed any messages, persist
+		if (persist) {
+			this.persistence.set(SLACK_MESSAGES_PERSISTENCE_KEY, persistedMessages)
+		}
+	}
+}
\ No newline at end of file
diff --git a/Engine/Extras/RoboMerge/v3/src/robo/preview.ts b/Engine/Extras/RoboMerge/v3/src/robo/preview.ts
--- a/Engine/Extras/RoboMerge/v3/src/robo/preview.ts
+++ b/Engine/Extras/RoboMerge/v3/src/robo/preview.ts
@@ -59,9 +59,7 @@ export async function getPreview(cl: number, singleBot?: string) {
 			continue
 		}
 
-		for (const branch of graph.branches) {
-			status.addBranch(branch)
-		}
+		await Promise.all(graph.branches.map(branch => status.addBranch(branch)))
 	}
 
 	if (errors.length > 0) {
diff --git a/Engine/Extras/RoboMerge/v3/src/robo/robo.ts b/Engine/Extras/RoboMerge/v3/src/robo/robo.ts
--- a/Engine/Extras/RoboMerge/v3/src/robo/robo.ts
+++ b/Engine/Extras/RoboMerge/v3/src/robo/robo.ts
@@ -26,6 +26,7 @@ import { notificationsInit } from './notifications'
 import { RoboArgs } from './roboargs'
 import { settingsInit } from './settings'
 import { Vault } from './vault'
+import { doDBPersistedMessagesHousekeeping } from './persistedmessages';
 
 /*************************
  * RoboMerge main process
@@ -234,12 +235,13 @@ settingsInit(args)
  * Connect to MongoDB.
  */
 if (RoboArgs.useMongo()) {
-	var mongoose = require('mongoose');
-	mongoose.connect(args.mongoDB_URI);
-	  mongoose.connection.on('error', (err: any) => {
-		console.log(`MongoDB Connection Error: ${err}`);
+	var mongoose = require('mongoose')
+	mongoose.connect(args.mongoDB_URI)
+	mongoose.connection.on('connected', () => doDBPersistedMessagesHousekeeping())
+	mongoose.connection.on('error', (err: any) => {
+		console.log(`MongoDB Connection Error: ${err}`)
 		process.exit(1);
-	  });
+	})
 }
 
 const env = args.epicEnv || (args.nodeEnv === "production" ? "prod" : "dev")
diff --git a/Engine/Extras/RoboMerge/v3/src/robo/settings.ts b/Engine/Extras/RoboMerge/v3/src/robo/settings.ts
--- a/Engine/Extras/RoboMerge/v3/src/robo/settings.ts
+++ b/Engine/Extras/RoboMerge/v3/src/robo/settings.ts
@@ -7,6 +7,7 @@ import { VersionReader } from '../common/version';
 import { PerforceContext, Workspace } from '../common/perforce';
 import { RoboArgs } from './roboargs';
 import semver = require('semver');
+import { migrateMessagesToDB } from './persistedmessages';
 
 const BotModel = require('./models/bot');
 
@@ -104,6 +105,7 @@ export class Settings {
 	botname: string
 	enableSave: boolean
 	object: { [key: string]: any }
+	private savingState: 'Idle' | 'Saving' | 'NeedsSave' = 'Idle'
 	private	lastPersistBackupTime: Date
 	private readonly settingsLogger: ContextualLogger
 	private readonly p4backupWorkspace?: Workspace
@@ -141,12 +143,15 @@ export class Settings {
 		}
 
 		let objVersion : semver.SemVer
+		let loadedFromFile = false
 		if (!settings.object)
 		{
 			// load the object from disk
 			let filebits = readFileToString(settings.getFilename('PERSISTENCE'));
-			if (filebits)
+			if (filebits) {
 				settings.object = jsonlint.parse(filebits);
+				loadedFromFile = true
+			}
 			else {
 				// Create "empty" settings object, but include latest version so we don't needless enter migration code
 				settings.object = {
@@ -182,6 +187,11 @@ export class Settings {
 		// For comparison later
 		const previousObjVersion = new semver.SemVer(objVersion)
 
+		// Transition slack messages out of the settings object and into the database
+		if (loadedFromFile && RoboArgs.useMongo()) {
+			await migrateMessagesToDB(settings)
+		}
+
 		// We're up to date!
 		if (previousObjVersion.compare(VERSION) !== 0) {
 			settings.settingsLogger.info(`Changing settings version from '${previousObjVersion.raw}' to '${VERSION.raw}'`)
@@ -207,9 +217,25 @@ export class Settings {
 
 	_saveObject(ensureInP4?: boolean) {
 		if (this.enableSave) {
-
 			if (RoboArgs.useMongo()) {
-				BotModel.collection.updateOne({botname:this.botname},{$set:{botname:this.botname,details:this.object}},{upsert:true})
+				const saveToDB = () => {
+					this.savingState = 'Saving'
+					BotModel.collection.updateOne({botname:this.botname},{$set:{botname:this.botname,details:this.object}},{upsert:true})
+					.then(() => {
+						if (this.savingState == 'NeedsSave') {
+							saveToDB()
+						}
+						else {
+							this.savingState = 'Idle'
+						}
+					})
+				}
+				if (this.savingState == 'Idle') {
+					saveToDB()
+				}
+				else {
+					this.savingState = 'NeedsSave'
+				}
 			}
 			else {
 				const persistentFile = this.getFilename('PERSISTENCE')
diff --git a/Engine/Extras/RoboMerge/v3/src/robo/status.ts b/Engine/Extras/RoboMerge/v3/src/robo/status.ts
--- a/Engine/Extras/RoboMerge/v3/src/robo/status.ts
+++ b/Engine/Extras/RoboMerge/v3/src/robo/status.ts
@@ -12,12 +12,12 @@ export class Status {
 		this.statusLogger = parentLogger.createChild('Status')
 	}
 
-	addBranch(branch: Branch) {
+	async addBranch(branch: Branch) {
 		const status: BranchStatus = {def: Status.makeBranchDef(branch), bot: branch.parent.botname, branch_spec_cl: -1}
 		this.allBranches.push(status)
 
 		if (branch.bot) {
-			branch.bot.applyStatus(status)
+			await branch.bot.applyStatus(status)
 		}
 	}
 
