Revert files with no diffs when using skip-if-unchanged as part of skipping the integration

diff --git a/Engine/Extras/RoboMerge/v3/src/robo/edgebot.ts b/Engine/Extras/RoboMerge/v3/src/robo/edgebot.ts
--- a/Engine/Extras/RoboMerge/v3/src/robo/edgebot.ts
+++ b/Engine/Extras/RoboMerge/v3/src/robo/edgebot.ts
@@ -798,6 +798,7 @@ class EdgeBotImpl extends PerforceStatefulBot {
 				}
 				
 				if (!foundDiff) {
+					await this.p4.revert(pending.change.targetWorkspace, pending.newCl, [], edgeServerAddress)
 					return await this.SkipIntegration(
 						pending,
 						edgeServerAddress,
