Bump the package version for DB transition and removal of backwards compatibility code

diff --git a/Engine/Extras/RoboMerge/v3/package.json b/Engine/Extras/RoboMerge/v3/package.json
--- a/Engine/Extras/RoboMerge/v3/package.json
+++ b/Engine/Extras/RoboMerge/v3/package.json
@@ -1,6 +1,6 @@
 {
   "name": "robomerge",
-  "version": "3.0.4",
+  "version": "3.1.0",
   "description": "Automatic P4 merge management bot.",
   "repository": {},
   "keywords": [
