Prompt whether or not the user really wants to unpause or not to reduce accidental unpauses

diff --git a/Engine/Extras/RoboMerge/v3/public/js/boilerplate.js b/Engine/Extras/RoboMerge/v3/public/js/boilerplate.js
--- a/Engine/Extras/RoboMerge/v3/public/js/boilerplate.js
+++ b/Engine/Extras/RoboMerge/v3/public/js/boilerplate.js
@@ -1367,15 +1367,18 @@ function renderActionsCell_Common(actionCell, data, operationFunction, operation
 		specificActionButton = createActionButton(
 			`Unpause`,
 			function() {
-				operationFunction(...operationArgs, '/unpause', function(success) {
-					if (success) {
-						displaySuccessfulMessage(`Unpaused ${dataDisplayName}.`)
-					} else {
-						displayErrorMessage(`Error unpausing ${dataFullName}, please check logs.`)
-						$('#helpButton').trigger('click')
-					}
-					updateBranchList(botname)
-				})
+				let unpauseConfirmation = confirm(`Are you sure you want to unpause ${dataDisplayName}?`)
+				if (unpauseConfirmation) {
+					operationFunction(...operationArgs, '/unpause', function(success) {
+						if (success) {
+							displaySuccessfulMessage(`Unpaused ${dataDisplayName}.`)
+						} else {
+							displayErrorMessage(`Error unpausing ${dataFullName}, please check logs.`)
+							$('#helpButton').trigger('click')
+						}
+						updateBranchList(botname)
+					})
+				}
 			},
 			`Resume Robomerge monitoring of ${dataDisplayName}`
 		)
