Identify transfer edges when tracking back in more cases

diff --git a/Engine/Extras/RoboMerge/v3/src/robo/trackchange.ts b/Engine/Extras/RoboMerge/v3/src/robo/trackchange.ts
--- a/Engine/Extras/RoboMerge/v3/src/robo/trackchange.ts
+++ b/Engine/Extras/RoboMerge/v3/src/robo/trackchange.ts
@@ -120,38 +120,51 @@ export async function trackChange(robo: RoboMerge, logger: ContextualLogger, ini
 
     let changes = new Map<string, TrackedChange>()
 
+    let determineSourceServer = (node: Node) => {
+        let edges = graph.getEdgesForNode(node)
+        const possibleServers = new Set([...edges].filter(e => e.target.stream === node.stream).map(e => e.source.serverID))
+        if (possibleServers.size == 1) {
+            return [...possibleServers][0]
+        } 
+        else {
+            // TODO: figure out which one it is
+        }
+        return undefined
+    }
+
     let sourceChain: TrackedCL[] = []
     let gatherCLInfo = async (cl: TrackedCL, opts?: GatherOptions) => {
 
         if (cl.serverId === undefined && opts?.lastCL) {
             const lastChange = changes.get(JSON.stringify(opts.lastCL))
             if (lastChange?.node) {
-                let edges = graph.getEdgesForNode(lastChange.node)
-                const possibleServers = new Set([...edges].filter(e => e.target.stream === lastChange.node!.stream).map(e => e.source.serverID))
-                if (possibleServers.size == 1) {
-                    cl.serverId = [...possibleServers][0]
-                } 
-                else {
-                    // TODO: figure out which one it is
-                }
-            }    
+                cl.serverId = determineSourceServer(lastChange.node)
+            }
         }
 
         let p4 = PerforceContext.getServerContext(robo.roboMergeLogger, cl.serverId)
         try {
             let desc = await p4.describe(cl.cl, 1)
+            const clNode = getNode(desc)
             let initialSubmit = false
             if (!opts?.parentCL) {
                 //TODO: abstract this in to a description parser per stream/branch/server
                 const sourceMatch = desc.description.match(/#ROBOMERGE-SOURCE: (?<RMSource>.*)\n|Merge-CL: (?<MergeCL>\d*)\n/)
                 if (sourceMatch?.groups!.RMSource) {
                     const newSources: TrackedCL[] = Array.from(sourceMatch.groups.RMSource.matchAll(/CL (\d+)/g)).map(m => {return {cl: parseInt(m[1])}}).reverse()
+                    if (clNode) {
+                        newSources[0].serverId = determineSourceServer(clNode)
+                    }
                     sourceChain.push(...newSources)
                     opts = { ...opts, parentCL: newSources[0] }
                 }
                 else if (sourceMatch?.groups!.MergeCL) {
-                    opts = { ...opts, parentCL: {cl:parseInt(sourceMatch.groups.MergeCL)}}
-                    sourceChain.push(opts.parentCL!)
+                    const parentCL = {
+                        cl: parseInt(sourceMatch.groups.MergeCL),
+                        serverId: clNode ? determineSourceServer(clNode) : undefined
+                    }
+                    opts = { ...opts, parentCL}
+                    sourceChain.push(parentCL)
                 }
                 else {
                     for (const entry of desc.entries) {
@@ -174,7 +187,6 @@ export async function trackChange(robo: RoboMerge, logger: ContextualLogger, ini
                 initialSubmit = opts?.parentCL === undefined
             }
             const mergeMethod = getMergeMethod(desc.description, cl, {...opts, initialSubmit} )
-            const clNode = getNode(desc)
             if (clNode || (RobomergeMethodStrings as readonly string[]).includes(mergeMethod)) {
                 changes.set(JSON.stringify(cl), {desc, mergeMethod, node: clNode, sourceCL: opts?.parentCL, destCLs: (opts && opts.lastCL ? [opts.lastCL] : []) })
             }
