Submit with TODO comment about failure cases

diff --git a/Engine/Extras/RoboMerge/v3/src/robo/trackchange.ts b/Engine/Extras/RoboMerge/v3/src/robo/trackchange.ts
--- a/Engine/Extras/RoboMerge/v3/src/robo/trackchange.ts
+++ b/Engine/Extras/RoboMerge/v3/src/robo/trackchange.ts
@@ -152,6 +152,7 @@ export async function trackChange(robo: RoboMerge, logger: ContextualLogger, ini
                 const sourceMatch = desc.description.match(/#ROBOMERGE-SOURCE: (?<RMSource>.*)\n|Merge-CL: (?<MergeCL>\d*)\n/)
                 if (sourceMatch?.groups!.RMSource) {
                     const newSources: TrackedCL[] = Array.from(sourceMatch.groups.RMSource.matchAll(/CL (\d+)/g)).map(m => {return {cl: parseInt(m[1])}}).reverse()
+                    // TODO: this only solves the first step, if the server transfer is 2nd or further back then this will still fail to properly identify it
                     if (clNode) {
                         newSources[0].serverId = determineSourceServer(clNode)
                     }
