Remove unused revertFile function
Make revert able to revert a specific file in the specified workspace

diff --git a/Engine/Extras/RoboMerge/v3/src/common/perforce.ts b/Engine/Extras/RoboMerge/v3/src/common/perforce.ts
--- a/Engine/Extras/RoboMerge/v3/src/common/perforce.ts
+++ b/Engine/Extras/RoboMerge/v3/src/common/perforce.ts
@@ -1578,22 +1578,9 @@ export class PerforceContext {
 		return this.execAndParse(roboWorkspace, args) as Promise<OpenedFileRecord[]>
 	}
 
-	async revertFile(file: string) {
-		let opened = await this.opened(null, file, true)
-		if (opened.length == 0) {
-			opened = await this.opened(null, file)
-			if (opened.length == 0) {
-				this.logger.warn(`${file} does not appear to be open, cannot revert`)
-				return null
-			}
-		}
-		const client = opened[0].client
-		return this.revertFiles([file],client)
-	}
-
-	async revertFiles(files: string[], client: string) {
+	async revertFiles(files: string[], client: string, opts?: string[]) {
 		const edgeServer = await this.getWorkspaceEdgeServer(client)
-		const args = ['revert', '-C', client, ...files]
+		const args = ['revert', ...(opts||[]), '-C', client, ...files]
 		try {
 			const results = await this.execAndParse(null, args, {serverAddress: edgeServer?.address})
 
@@ -1627,13 +1614,22 @@ export class PerforceContext {
 
 	// revert a CL deleting any files marked for add
 	// output format is just error or not
-	async revert(roboWorkspace: RoboWorkspace, changelist: number, opts?: string[], edgeServerAddress?: string) {
+	async revert(roboWorkspace: RoboWorkspace, changelist: number, opts?: string[], edgeServerAddress?: string): Promise<void>
+	async revert(roboWorkspace: RoboWorkspace, file: string, opts?: string[], edgeServerAddress?: string): Promise<void>
+	async revert(roboWorkspace: RoboWorkspace, arg: number|string, opts?: string[], edgeServerAddress?: string) {
 		const workspace = coercePerforceWorkspace(roboWorkspace);
-		const path = workspace && workspace.name ? `//${workspace.name}/...` : ''
 		try {
-			await this._execP4(workspace, ['revert',
-				/*"-w", seems to cause ENOENT occasionally*/ ...(opts || []), '-c', changelist.toString(), path],
-				{serverAddress: edgeServerAddress});
+			if (typeof arg === 'number') {
+				const path = workspace && workspace.name ? `//${workspace.name}/...` : ''
+				await this._execP4(workspace, ['revert',
+					/*"-w", seems to cause ENOENT occasionally*/ ...(opts || []), '-c', arg.toString(), path],
+					{serverAddress: edgeServerAddress});
+			}
+			else {
+				await this._execP4(workspace, ['revert',
+					/*"-w", seems to cause ENOENT occasionally*/ ...(opts || []), arg],
+					{serverAddress: edgeServerAddress});
+			}
 		}
 		catch (reason) {
 			if (!isExecP4Error(reason)) {
