Fix crash when going to update a persistence file that is already opened for edit for some reason (most likely a server crash or perforce error while it was in the process of updating previously)

diff --git a/Engine/Extras/RoboMerge/v3/src/robo/settings.ts b/Engine/Extras/RoboMerge/v3/src/robo/settings.ts
--- a/Engine/Extras/RoboMerge/v3/src/robo/settings.ts
+++ b/Engine/Extras/RoboMerge/v3/src/robo/settings.ts
@@ -266,7 +266,8 @@ export class Settings {
 					else if ((Date.now() - this.lastPersistBackupTime.getTime()) / 60000 >= args.persistenceBackupFrequency) {
 						this.lastPersistBackupTime = new Date(Date.now())
 
-						this.p4.sync(this.p4backupWorkspace, persistentFile, {opts: ['-k']})
+						this.p4.revert(this.p4backupWorkspace, persistentFile, ['-k'])
+						.then(() => this.p4.sync(this.p4backupWorkspace, persistentFile, {opts: ['-k']}))
 						.then(() => this.p4.new_cl(this.p4backupWorkspace!, "Updating persistent backup"))
 						.then((cl: number) => {
 							this.p4.edit(this.p4backupWorkspace!, cl, persistentFile, ['-k'])
