Don't overload revert and have an explicit revertFile as there seems to be cases where the "number" CL is actually a string

diff --git a/Engine/Extras/RoboMerge/v3/package.json b/Engine/Extras/RoboMerge/v3/package.json
--- a/Engine/Extras/RoboMerge/v3/package.json
+++ b/Engine/Extras/RoboMerge/v3/package.json
@@ -1,6 +1,6 @@
 {
   "name": "robomerge",
-  "version": "3.1.0",
+  "version": "3.1.1",
   "description": "Automatic P4 merge management bot.",
   "repository": {},
   "keywords": [
diff --git a/Engine/Extras/RoboMerge/v3/src/common/perforce.ts b/Engine/Extras/RoboMerge/v3/src/common/perforce.ts
--- a/Engine/Extras/RoboMerge/v3/src/common/perforce.ts
+++ b/Engine/Extras/RoboMerge/v3/src/common/perforce.ts
@@ -1614,22 +1614,10 @@ export class PerforceContext {
 
 	// revert a CL deleting any files marked for add
 	// output format is just error or not
-	async revert(roboWorkspace: RoboWorkspace, changelist: number, opts?: string[], edgeServerAddress?: string): Promise<void>
-	async revert(roboWorkspace: RoboWorkspace, file: string, opts?: string[], edgeServerAddress?: string): Promise<void>
-	async revert(roboWorkspace: RoboWorkspace, arg: number|string, opts?: string[], edgeServerAddress?: string) {
+	private async _revert(roboWorkspace: RoboWorkspace, revertArgs: string[], edgeServerAddress?: string) {
 		const workspace = coercePerforceWorkspace(roboWorkspace);
 		try {
-			if (typeof arg === 'number') {
-				const path = workspace && workspace.name ? `//${workspace.name}/...` : ''
-				await this._execP4(workspace, ['revert',
-					/*"-w", seems to cause ENOENT occasionally*/ ...(opts || []), '-c', arg.toString(), path],
-					{serverAddress: edgeServerAddress});
-			}
-			else {
-				await this._execP4(workspace, ['revert',
-					/*"-w", seems to cause ENOENT occasionally*/ ...(opts || []), arg],
-					{serverAddress: edgeServerAddress});
-			}
+			await this._execP4(workspace, revertArgs, {serverAddress: edgeServerAddress});
 		}
 		catch (reason) {
 			if (!isExecP4Error(reason)) {
@@ -1645,6 +1633,18 @@ export class PerforceContext {
 		}
 	}
 
+	revertFile(roboWorkspace: RoboWorkspace, file: string, opts?: string[], edgeServerAddress?: string) {
+		return this._revert(roboWorkspace, ['revert', ...(opts || []), file], edgeServerAddress)
+	}
+
+	// revert a CL deleting any files marked for add
+	// output format is just error or not
+	revert(roboWorkspace: RoboWorkspace, changelist: number, opts?: string[], edgeServerAddress?: string) {
+		const workspace = coercePerforceWorkspace(roboWorkspace);
+		const path = workspace && workspace.name ? `//${workspace.name}/...` : ''
+		return this._revert(roboWorkspace, ['revert', ...(opts || []), '-c', changelist.toString(), path], edgeServerAddress)
+	}
+
 	revertAndDelete(roboWorkspace: string, cl: number, edgeServerAddress?: string) {
 		return this.revert(roboWorkspace, cl, [], edgeServerAddress)
 		.then(() => {
diff --git a/Engine/Extras/RoboMerge/v3/src/robo/settings.ts b/Engine/Extras/RoboMerge/v3/src/robo/settings.ts
--- a/Engine/Extras/RoboMerge/v3/src/robo/settings.ts
+++ b/Engine/Extras/RoboMerge/v3/src/robo/settings.ts
@@ -266,7 +266,7 @@ export class Settings {
 					else if ((Date.now() - this.lastPersistBackupTime.getTime()) / 60000 >= args.persistenceBackupFrequency) {
 						this.lastPersistBackupTime = new Date(Date.now())
 
-						this.p4.revert(this.p4backupWorkspace, persistentFile, ['-k'])
+						this.p4.revertFile(this.p4backupWorkspace, persistentFile, ['-k'])
 						.then(() => this.p4.sync(this.p4backupWorkspace, persistentFile, {opts: ['-k']}))
 						.then(() => this.p4.new_cl(this.p4backupWorkspace!, "Updating persistent backup"))
 						.then((cl: number) => {
