Fix sunsetting DB persisted slack messages

diff --git a/Engine/Extras/RoboMerge/v3/src/robo/persistedmessages.ts b/Engine/Extras/RoboMerge/v3/src/robo/persistedmessages.ts
--- a/Engine/Extras/RoboMerge/v3/src/robo/persistedmessages.ts
+++ b/Engine/Extras/RoboMerge/v3/src/robo/persistedmessages.ts
@@ -188,7 +188,7 @@ class DBPersistedSlackMessages extends PersistedSlackMessages {
 }
 
 export async function doDBPersistedMessagesHousekeeping() {
-	const expiredTime = Date.now() - (SUNSET_PERSISTED_NOTIFICATIONS_DAYS * 60 * 60 * 24 * 1000)
+	const expiredTime = (Date.now() / 1000) - (SUNSET_PERSISTED_NOTIFICATIONS_DAYS * 60 * 60 * 24)
 	return SlackMessagesModel.collection.deleteMany({timestamp_forSunset: { $lt: expiredTime }})
 }
 
