Fix finding DB persisted slack messages

diff --git a/Engine/Extras/RoboMerge/v3/src/robo/persistedmessages.ts b/Engine/Extras/RoboMerge/v3/src/robo/persistedmessages.ts
--- a/Engine/Extras/RoboMerge/v3/src/robo/persistedmessages.ts
+++ b/Engine/Extras/RoboMerge/v3/src/robo/persistedmessages.ts
@@ -149,7 +149,11 @@ class DBPersistedSlackMessages extends PersistedSlackMessages {
 		}
 		const messageRecord = await SlackMessagesModel.find(index)
 		if (messageRecord.length == 1) {
-			const persistedMessage = new PersistedSlackMessage(messageRecord.message, key, this)
+			const persistedMessage = new PersistedSlackMessage({
+						messageOpts: messageRecord.message, 
+						permalink: messageRecord.permalink,
+						timestamp: messageRecord.timestamp
+					}, key, this)
 			this.messageCache[cacheKey] = persistedMessage
 			return persistedMessage
 		}
@@ -173,7 +177,11 @@ class DBPersistedSlackMessages extends PersistedSlackMessages {
 			const key = new PersistedMessageKey(cl, index.branchName, messageRecord.channel)
 			const cacheKey = key.toString()
 			if (!this.messageCache[cacheKey]) {
-				const persistedMessage = new PersistedSlackMessage(messageRecord.message, key, this)
+				const persistedMessage = new PersistedSlackMessage({
+						messageOpts: messageRecord.message, 
+						permalink: messageRecord.permalink,
+						timestamp: messageRecord.timestamp
+					}, key, this)
 				this.messageCache[cacheKey] = persistedMessage
 			}
 			messages.push(this.messageCache[cacheKey])
