Get workspaces on the appropriate server when requesting a shelf

diff --git a/Engine/Extras/RoboMerge/v3/public/js/boilerplate.js b/Engine/Extras/RoboMerge/v3/public/js/boilerplate.js
--- a/Engine/Extras/RoboMerge/v3/public/js/boilerplate.js
+++ b/Engine/Extras/RoboMerge/v3/public/js/boilerplate.js
@@ -303,10 +303,10 @@ function getBranch(botname, branchname, callback) {
 	});
 }
 
-function getUserWorkspaces(callback) {
+function getUserWorkspaces(serverID, callback) {
 	clearErrorText();
 	$.ajax({
-		url: `/api/user/workspaces`,
+		url: `/api/user/${serverID}/workspaces`,
 		type: 'get',
 		dataType: 'json',
 		success: callback,
diff --git a/Engine/Extras/RoboMerge/v3/public/js/createshelf.js b/Engine/Extras/RoboMerge/v3/public/js/createshelf.js
--- a/Engine/Extras/RoboMerge/v3/public/js/createshelf.js
+++ b/Engine/Extras/RoboMerge/v3/public/js/createshelf.js
@@ -43,7 +43,7 @@ function shelfVerify() {
 
         // Determine if we have any relevant workspaces
         let foundWorkspaces = false
-        getUserWorkspaces(function(wsdata) {
+        getUserWorkspaces(requestedEdgeData.serverID, function(wsdata) {
             // Ensure we have data
             if (!wsdata || !wsdata.data) {
                 shelfFailure('Workspaces could not be found for user.')
@@ -57,7 +57,7 @@ function shelfVerify() {
             }
 
             let userWorkspaces = wsdata.data
-            if (!userWorkspaces || userWorkspaces === []) {
+            if (!userWorkspaces || userWorkspaces.length == 0) {
                 shelfFailure("Found no user workspaces.")
                 return
             }
diff --git a/Engine/Extras/RoboMerge/v3/src/robo/edgebot.ts b/Engine/Extras/RoboMerge/v3/src/robo/edgebot.ts
--- a/Engine/Extras/RoboMerge/v3/src/robo/edgebot.ts
+++ b/Engine/Extras/RoboMerge/v3/src/robo/edgebot.ts
@@ -1098,6 +1098,7 @@ class EdgeBotImpl extends PerforceStatefulBot {
 		status.rootPath = this.targetBranch.rootPath
 
 		status.last_cl = this.lastCl
+		status.serverID = this.p4.serverID
 		status.swarmURL = this.p4.swarmURL
 		
 		status.is_active = this.isActive
diff --git a/Engine/Extras/RoboMerge/v3/src/robo/ipc.ts b/Engine/Extras/RoboMerge/v3/src/robo/ipc.ts
--- a/Engine/Extras/RoboMerge/v3/src/robo/ipc.ts
+++ b/Engine/Extras/RoboMerge/v3/src/robo/ipc.ts
@@ -55,7 +55,7 @@ export class IPC {
 			case 'setVerbose': return this.setVerbose(!!msg.args![0])
 			case 'setVerbosity': return this.setVerbosity(msg.args![0] as string)
 			case 'getp4tasks': return { ...OPERATION_SUCCESS, data: getRunningPerforceCommands() } 
-			case 'getWorkspaces': return this.getWorkspaces(msg.args![0] as string)
+			case 'getWorkspaces': return this.getWorkspaces(msg.args![0] as string, msg.args![1] as string)
 			case 'getPersistence': return this.getPersistence(msg.args![0] as string)
 			case 'traceRoute': return this.traceRoute(msg.args![0] as Query, msg.args![1])
 			case 'trackChange': return this.trackChange(msg.args![0], msg.args![1])
@@ -146,8 +146,8 @@ export class IPC {
 		}
 	}
 
-	private async getWorkspaces(username: string): Promise<OperationReturnType> {
-		const workspaces = await p4util.getWorkspacesForUser(PerforceContext.getServerContext(this.robo.roboMergeLogger), username)
+	private async getWorkspaces(username: string, serverID: string): Promise<OperationReturnType> {
+		const workspaces = await p4util.getWorkspacesForUser(PerforceContext.getServerContext(this.robo.roboMergeLogger, serverID), username)
 		
 		if (!workspaces) {
 			return {
diff --git a/Engine/Extras/RoboMerge/v3/src/robo/roboserver.ts b/Engine/Extras/RoboMerge/v3/src/robo/roboserver.ts
--- a/Engine/Extras/RoboMerge/v3/src/robo/roboserver.ts
+++ b/Engine/Extras/RoboMerge/v3/src/robo/roboserver.ts
@@ -454,13 +454,13 @@ class RoboWebApp implements AppInterface {
 		return await this.sendMessage('setVerbosity', [level.toLowerCase()])
 	}
 
-	@SecureHandler('GET', '/api/user/workspaces')
-	async getWorkspaces(): Promise<OperationReturnType> {
+	@SecureHandler('GET', '/api/user/*/workspaces')
+	async getWorkspaces(serverID: string): Promise<OperationReturnType> {
 		if (!this.authData) {
 			throw new Error('Secure call but no auth data?')
 		}
 
-		return await this.sendMessage('getWorkspaces', [this.authData.user])
+		return await this.sendMessage('getWorkspaces', [this.authData.user, serverID])
 	}
 
 	@SecureHandler('GET', '/api/branches')
diff --git a/Engine/Extras/RoboMerge/v3/src/robo/status-types.ts b/Engine/Extras/RoboMerge/v3/src/robo/status-types.ts
--- a/Engine/Extras/RoboMerge/v3/src/robo/status-types.ts
+++ b/Engine/Extras/RoboMerge/v3/src/robo/status-types.ts
@@ -121,6 +121,8 @@ export type EdgeStatusFields = BotStatusFields & {
 	lastGoodCLDate?: Date
 	waitingForCISLink?: string
 
+	serverID: string
+
 	num_changes_remaining: number
 }
 
