Improve track changes when dealing with non-streams

diff --git a/Engine/Extras/RoboMerge/v3/src/common/helper.ts b/Engine/Extras/RoboMerge/v3/src/common/helper.ts
--- a/Engine/Extras/RoboMerge/v3/src/common/helper.ts
+++ b/Engine/Extras/RoboMerge/v3/src/common/helper.ts
@@ -72,6 +72,16 @@ export function sortBy<T>(arr: T[], projection: (x: T) => number|string) {
 	return decorated.map(p => p[1])
 }
 
+export function matchPrefix(a: string, b: string) {
+	const len = Math.min(a.length, b.length)
+	for (let i = 0; i < len; ++i) {
+		if (a.charAt(i) !== b.charAt(i)) {
+			return i
+		}
+	}
+	return len
+}
+
 export class Random {
 	static randint(minInc: number, maxInc: number) {
 		return Math.floor(Math.random() * (maxInc - minInc)) + minInc
diff --git a/Engine/Extras/RoboMerge/v3/src/common/p4util.ts b/Engine/Extras/RoboMerge/v3/src/common/p4util.ts
--- a/Engine/Extras/RoboMerge/v3/src/common/p4util.ts
+++ b/Engine/Extras/RoboMerge/v3/src/common/p4util.ts
@@ -1,6 +1,6 @@
 // Copyright Epic Games, Inc. All Rights Reserved.
 
-import { _setTimeout } from './helper';
+import { matchPrefix, _setTimeout } from './helper';
 import { ContextualLogger } from './logger';
 import { Change, ClientSpec, PerforceContext, Workspace } from './perforce';
 
@@ -89,16 +89,6 @@ export async function getWorkspacesForUser(p4: PerforceContext, user: string) {
 		.filter(ws => !USER_WORKSPACE_EXCLUDE_PATTERNS.some(entry => ws.client.match(entry)))
 }
 
-function matchPrefix(a: string, b: string) {
-	const len = Math.min(a.length, b.length)
-	for (let i = 0; i < len; ++i) {
-		if (a.charAt(i) !== b.charAt(i)) {
-			return i
-		}
-	}
-	return len
-}
-
 export async function chooseBestWorkspaceForUser(p4: PerforceContext, user: string, stream?: string) {
 	const workspaces: ClientSpec[] = await getWorkspacesForUser(p4,user)
 
diff --git a/Engine/Extras/RoboMerge/v3/src/robo/graph.ts b/Engine/Extras/RoboMerge/v3/src/robo/graph.ts
--- a/Engine/Extras/RoboMerge/v3/src/robo/graph.ts
+++ b/Engine/Extras/RoboMerge/v3/src/robo/graph.ts
@@ -1,6 +1,6 @@
 // Copyright Epic Games, Inc. All Rights Reserved.
 
-import { setDefault, sortBy } from '../common/helper';
+import { matchPrefix, setDefault, sortBy } from '../common/helper';
 import { ContextualLogger } from '../common/logger';
 import { PerforceContext } from '../common/perforce';
 import { Branch, BranchGraphInterface } from '../robo/branch-interfaces';
@@ -130,15 +130,17 @@ export class Graph {
 		return setDefault(this.streamNodes, branchId, new Node(branchId, server))
 	}
 
-	findNodeForStream(depotPath: string | Branch, allowPartialMatch?: boolean) {
+	findNodeForStream(depotPath: string | Branch, partialMatchPath?: string) {
 		const branchId = makeBranchId(depotPath)
 		let nodeForStream = this.streamNodes.get(branchId)
-		if (!nodeForStream && allowPartialMatch) {
-			const partialMatchbranchId = (branchId as string) + '/'
+		if (!nodeForStream && partialMatchPath) {
+
+			let target_match = 0
 			for (const [stream, node] of this.streamNodes) {
-				if (stream.startsWith(partialMatchbranchId)) {
+				const matchlen = matchPrefix(stream.toLowerCase(), partialMatchPath.toLowerCase())
+				if (matchlen > target_match) {
+					target_match = matchlen
 					nodeForStream = node
-					break
 				}
 			}
 		}
diff --git a/Engine/Extras/RoboMerge/v3/src/robo/trackchange.ts b/Engine/Extras/RoboMerge/v3/src/robo/trackchange.ts
--- a/Engine/Extras/RoboMerge/v3/src/robo/trackchange.ts
+++ b/Engine/Extras/RoboMerge/v3/src/robo/trackchange.ts
@@ -71,7 +71,7 @@ export async function trackChange(robo: RoboMerge, logger: ContextualLogger, ini
             const stream = getStreamFromPath(desc.path)
             if (stream) {
                 try {
-                    return graph.findNodeForStream(stream, true)
+                    return graph.findNodeForStream(stream, desc.path)
                 }
                 catch {						
                 }
