Fix infinite loop when the length of dirs is 0

diff --git a/Engine/Extras/RoboMerge/v3/src/robo/trackchange.ts b/Engine/Extras/RoboMerge/v3/src/robo/trackchange.ts
--- a/Engine/Extras/RoboMerge/v3/src/robo/trackchange.ts
+++ b/Engine/Extras/RoboMerge/v3/src/robo/trackchange.ts
@@ -101,16 +101,23 @@ export async function trackChange(robo: RoboMerge, logger: ContextualLogger, ini
         // of files, otherwise just get an additional block to evaluate
         if (desc.path.endsWith("/...")) {
             const rootDir = `${desc.path.slice(0,-3)}*@=${clToConsider.cl}`
+            const FILES_PER_DIR = 5
+            const DIRS_PER_SLICE = 5
+            const filesPromises = [p4.files(rootDir, FILES_PER_DIR)]
+            const dirs = await p4.dirs(rootDir)
             let fpIndex = 0
-            let filesPromises = [p4.files(rootDir, 5)]
-            while (desc.entries.length <= 1) {
-                const dirs = await p4.dirs(rootDir)
-                filesPromises.push(...dirs.slice(0, 5).map(dir => p4.files(`${dir}/...@=${clToConsider.cl}`, 5)))
+            let dirSliceBegin = 0
+            while (true) {
+                filesPromises.push(...dirs.slice(dirSliceBegin, dirSliceBegin+DIRS_PER_SLICE).map(dir => p4.files(`${dir}/...@=${clToConsider.cl}`, FILES_PER_DIR)))
                 while (fpIndex < filesPromises.length) {
                     desc.entries.push(...(await filesPromises[fpIndex]))
                     fpIndex++
                 }
                 desc.entries = desc.entries.filter(e => e.action != 'move/delete')
+                dirSliceBegin += DIRS_PER_SLICE
+                if (desc.entries.length > 1 || dirSliceBegin >= dirs.length) {
+                    break
+                }
             }
         }
         else {
