Fix partial matches returning completely unrelated nodes

diff --git a/Engine/Extras/RoboMerge/v3/src/robo/graph.ts b/Engine/Extras/RoboMerge/v3/src/robo/graph.ts
--- a/Engine/Extras/RoboMerge/v3/src/robo/graph.ts
+++ b/Engine/Extras/RoboMerge/v3/src/robo/graph.ts
@@ -135,7 +135,9 @@ export class Graph {
 		let nodeForStream = this.streamNodes.get(branchId)
 		if (!nodeForStream && partialMatchPath) {
 
-			let target_match = 0
+			// don't match unless there is a match of at least as long as the root depotPath
+			const minMatchLength = (branchId as string).length
+			let target_match = minMatchLength-1
 			for (const [stream, node] of this.streamNodes) {
 				const matchlen = matchPrefix(stream.toLowerCase(), partialMatchPath.toLowerCase())
 				if (matchlen > target_match) {
