Error message for no file(s) is different on newer versions of p4 server

diff --git a/Engine/Extras/RoboMerge/v3/package.json b/Engine/Extras/RoboMerge/v3/package.json
--- a/Engine/Extras/RoboMerge/v3/package.json
+++ b/Engine/Extras/RoboMerge/v3/package.json
@@ -1,6 +1,6 @@
 {
   "name": "robomerge",
-  "version": "3.1.3",
+  "version": "3.1.4",
   "description": "Automatic P4 merge management bot.",
   "repository": {},
   "keywords": [
diff --git a/Engine/Extras/RoboMerge/v3/src/common/perforce.ts b/Engine/Extras/RoboMerge/v3/src/common/perforce.ts
--- a/Engine/Extras/RoboMerge/v3/src/common/perforce.ts
+++ b/Engine/Extras/RoboMerge/v3/src/common/perforce.ts
@@ -1830,7 +1830,8 @@ export class PerforceContext {
 
 			let [_, output] = reason
 			// If perforce doesn't detect revisions in the given range, return an empty set of revisions
-			if (output.trim().endsWith("no such file(s).")) {
+			if (output.trim().endsWith("no such file(s).") ||
+			    output.trim().endsWith("no file(s) at that changelist number.")) {
 				return []
 			}
 			throw reason		
