Learn to trust

diff --git a/Engine/Extras/RoboMerge/v3/src/common/perforce.ts b/Engine/Extras/RoboMerge/v3/src/common/perforce.ts
--- a/Engine/Extras/RoboMerge/v3/src/common/perforce.ts
+++ b/Engine/Extras/RoboMerge/v3/src/common/perforce.ts
@@ -509,6 +509,12 @@ export class PerforceContext {
 		const serverAddress = serverDetails && serverDetails.port
 
 		let opts: ExecOpts = { serverAddress }
+
+		if (serverAddress.startsWith("ssl:")) {
+			const trustArgs = ["trust", '-y']
+			await PerforceContext.execAndParse(logger, null, trustArgs, opts)
+		}
+
 		let loginArgs = ["login"]
 		if (serverDetails && serverDetails.user) {
 			loginArgs = ["-u", serverDetails.user, ...loginArgs]
