Fix null dereference

diff --git a/Engine/Extras/RoboMerge/v3/src/common/perforce.ts b/Engine/Extras/RoboMerge/v3/src/common/perforce.ts
--- a/Engine/Extras/RoboMerge/v3/src/common/perforce.ts
+++ b/Engine/Extras/RoboMerge/v3/src/common/perforce.ts
@@ -510,7 +510,7 @@ export class PerforceContext {
 
 		let opts: ExecOpts = { serverAddress }
 
-		if (serverAddress.startsWith("ssl:")) {
+		if (serverAddress?.startsWith("ssl:")) {
 			const trustArgs = ["trust", '-y']
 			await PerforceContext.execAndParse(logger, null, trustArgs, opts)
 		}
