Add missing return

diff --git a/Engine/Extras/RoboMerge/v3/src/common/perforce.ts b/Engine/Extras/RoboMerge/v3/src/common/perforce.ts
--- a/Engine/Extras/RoboMerge/v3/src/common/perforce.ts
+++ b/Engine/Extras/RoboMerge/v3/src/common/perforce.ts
@@ -1123,6 +1123,7 @@ export class PerforceContext {
 		const HandleCaughtError = (reason: any, filterOut?: (s: string) => boolean) => {
 			if (!isExecP4Error(reason)) {
 				failures.push(reason)
+				return
 			}
 
 			const [_, output] = reason
