When transferring respect the stream view of the target stream and filter out files that do not exist in the target view

diff --git a/Engine/Extras/RoboMerge/v3/src/common/p4util.ts b/Engine/Extras/RoboMerge/v3/src/common/p4util.ts
--- a/Engine/Extras/RoboMerge/v3/src/common/p4util.ts
+++ b/Engine/Extras/RoboMerge/v3/src/common/p4util.ts
@@ -6,10 +6,11 @@ import { Change, ClientSpec, PerforceContext, Workspace } from './perforce';
 
 
 export function makeRegexFromPerforcePath(path: string) {
+	path = path.replace('%%1', '*')
 	path = path.replace(/[/\-\\^$*+?.()|[\]{}]/g, '\\$&') // re.escape
-	path = path.replace('\\*', '[^/]*')
+	path = path.replace('\\*', '[^\/]*')
 	path = path.replace('\\.\\.\\.', '.*')
-	return new RegExp(path)
+	return new RegExp(`^${path}$`)
 }
 
 const USER_WORKSPACE_EXCLUDE_PATTERNS: (RegExp | string)[] = [
diff --git a/Engine/Extras/RoboMerge/v3/src/common/perforce.ts b/Engine/Extras/RoboMerge/v3/src/common/perforce.ts
--- a/Engine/Extras/RoboMerge/v3/src/common/perforce.ts
+++ b/Engine/Extras/RoboMerge/v3/src/common/perforce.ts
@@ -772,6 +772,16 @@ export class PerforceContext {
 		}
 	}
 
+	async streamView(streamName: string) {
+		if (await this.stream(streamName)) {
+			const streamResult = (await this.execAndParseArray(null, ['stream', '-ov', streamName]))[0]
+			return streamResult.entries.filter((result: any) => result.View).map((result: any) => result.View)
+		}
+		else {
+			return []
+		}
+	}
+
 	async branchExists(branchName: string) {
 		try {
 			return (await this.execAndParse(null, ['branches','-e',branchName])).length == 1
@@ -1072,7 +1082,36 @@ export class PerforceContext {
 
 		const p4quiet = !PerforceContext.devMode
 		const sourceStreamLength = source.stream!.length
-		const targetFilter = `${target.stream!}/...` == target.path_to ? null : makeRegexFromPerforcePath(target.path_to)
+
+		const buildFilter = (view: string[]) => {
+			if (view.length == 0) return null
+			return view.map(v => {
+				let path = v.split(' ')[0]
+				const exclude = path[0] == '-'
+				if (exclude) {
+					path = path.slice(1)
+				}
+				return {exclude, path: makeRegexFromPerforcePath(path) }
+			})
+		}
+
+		const evaluateFilter = (filter: any[], path: string) => {
+			let exclude = true
+			for (const f of filter) {
+				if (exclude != f.exclude && path.match(f.path)) {
+					exclude = !exclude
+				}
+			}
+			return exclude
+		}
+
+		let targetFilter: any = null
+		if (`${target.stream!}/...` == target.path_to) {
+			targetFilter = buildFilter(await targetP4.streamView(target.stream!))
+		}
+		else {
+			targetFilter = buildFilter([target.path_to])
+		}
 
 		let sortedFiles = new Map<string, any[]>()
 		const addSortedFile = (key: string, file: any) => {
@@ -1084,7 +1123,7 @@ export class PerforceContext {
 
 		files.forEach(file => {
 			const targetDepotFile = target.stream! + file.depotFile.substring(sourceStreamLength)
-			if (targetFilter && !targetDepotFile.match(targetFilter)) {
+			if (targetFilter && evaluateFilter(targetFilter, targetDepotFile)) {
 				// filter out files that don't match the target.path_to
 				return
 			}
