Fix error calling changes when CL description is empty

diff --git a/Engine/Extras/RoboMerge/v3/src/common/perforce.ts b/Engine/Extras/RoboMerge/v3/src/common/perforce.ts
--- a/Engine/Extras/RoboMerge/v3/src/common/perforce.ts
+++ b/Engine/Extras/RoboMerge/v3/src/common/perforce.ts
@@ -2283,7 +2283,7 @@ export class PerforceContext {
 			const parsedResult = JSON.parse(rawResult.slice(startIndex, endIndex != -1 ? endIndex+1 : undefined), reviver)
 			for (const expected in ((parseOptions && parseOptions.expected) || []))
 			{
-				if (!parsedResult[expected])
+				if (!(expected in parsedResult))
 				{
 					throw new Error(`Expected field ${expected} not present in ${JSON.stringify(parsedResult)}`)
 				}
