Error handling cleanup

diff --git a/Engine/Extras/RoboMerge/v3/public/js/acknowledge.js b/Engine/Extras/RoboMerge/v3/public/js/acknowledge.js
--- a/Engine/Extras/RoboMerge/v3/public/js/acknowledge.js
+++ b/Engine/Extras/RoboMerge/v3/public/js/acknowledge.js
@@ -44,7 +44,7 @@ function acknowledge() {
         
             // Ensure this node actually has this requested edge
             if (!requestedEdgeData) {
-                acknowledgeFailure(`${requestedBotName}:${requestedNodeName} has no edge for "${requestedEdgeName}".`)
+                acknowledgeFailure(`${requestedBotName}:${requestedBranchName} has no edge for "${requestedEdgeName}".`)
                 $('#result').append(renderSingleBranchTable(requestedBranchData))
                 return
             }
diff --git a/Engine/Extras/RoboMerge/v3/public/js/skip.js b/Engine/Extras/RoboMerge/v3/public/js/skip.js
--- a/Engine/Extras/RoboMerge/v3/public/js/skip.js
+++ b/Engine/Extras/RoboMerge/v3/public/js/skip.js
@@ -130,10 +130,10 @@ function skipVerify() {
 
 // This function takes in requested bot/branch/cl combo (usually from query parameters processed earlier) and 
 // ensures this is a valid request
-function verifySkipRequest(requestedBotName, requestedNodeName, requestedEdgeName, requestedBranchCl, requestedBranchData) {
+function verifySkipRequest(requestedBotName, requestedBranchName, requestedEdgeName, requestedBranchCl, requestedBranchData) {
     // Verify we got a branch
     if (!requestedBranchData.def) {
-        let errText = `Could not find matching branch for ${requestedBotName}:${requestedNodeName}.`
+        let errText = `Could not find matching branch for ${requestedBotName}:${requestedBranchName}.`
         if (requestedBranchData.message) {
             errText += `\n${requestedBranchData.message}`
         }
@@ -150,7 +150,7 @@ function verifySkipRequest(requestedBotName, requestedNodeName, requestedEdgeNam
 
     // Ensure this node actually has this requested edge
     if (!requestedEdgeData) {
-        skipFailure(`${requestedBotName}:${requestedNodeName} has no edge for "${requestedEdgeName}".`)
+        skipFailure(`${requestedBotName}:${requestedBranchName} has no edge for "${requestedEdgeName}".`)
         $('#result').append(renderSingleBranchTable(requestedBranchData))
         return
     }
diff --git a/Engine/Extras/RoboMerge/v3/public/js/stomp.js b/Engine/Extras/RoboMerge/v3/public/js/stomp.js
--- a/Engine/Extras/RoboMerge/v3/public/js/stomp.js
+++ b/Engine/Extras/RoboMerge/v3/public/js/stomp.js
@@ -18,6 +18,11 @@ function stompVerify() {
         return
     }
 
+    if (!requestedBotName || !requestedBranchName || !targetBranchName) {
+        stompFailure()
+        return
+    }
+
     // Filled in if/when user clicks the stomp button
     let stompOperation = $.Deferred()
     
@@ -46,7 +51,7 @@ function stompVerify() {
 
             // Ensure this node actually has this requested edge
             if (!requestedEdge) {
-                stompFailure(`${requestedBotName}:${requestedNodeName} has no edge for "${targetBranchName}".`)
+                stompFailure(`${requestedBotName}:${requestedBranchName} has no edge for "${targetBranchName}".`)
                 $('#result').append(renderSingleBranchTable(requestedNode))
                 return
             }
@@ -155,7 +160,7 @@ function stompVerify() {
             })
         }
         catch (err) {
-            stompFailure(`Error encountered displaying stomp verification results: ${err}`)
+            stompFailure(`Error encountered displaying stomp verification results: ${err}\n${err.stack}`)
         }
     })
 
@@ -333,13 +338,17 @@ function visualizeStompVerification(requestedBranchCl, stompJson) {
 }
 
 function stompFailure(message) {
-    displayErrorMessage(message, false)
+    if (message) {
+        displayErrorMessage(message, false)
+    }
     $('#loadingDiv').fadeOut("fast", "swing")
     $('#returnbutton').removeClass("initiallyHidden")
 }
 
 function stompSuccess(message) {
-    displaySuccessfulMessage(message, false)
+    if (message) {
+        displaySuccessfulMessage(message, false)
+    }
     $('#loadingDiv').fadeOut("fast", "swing")
     $('#returnbutton').removeClass("initiallyHidden")
 }
\ No newline at end of file
diff --git a/Engine/Extras/RoboMerge/v3/public/js/unlock.js b/Engine/Extras/RoboMerge/v3/public/js/unlock.js
--- a/Engine/Extras/RoboMerge/v3/public/js/unlock.js
+++ b/Engine/Extras/RoboMerge/v3/public/js/unlock.js
@@ -46,7 +46,7 @@ function unlockVerify() {
 
             // Ensure this node actually has this requested edge
             if (!requestedEdge) {
-                unlockFailure(`${requestedBotName}:${requestedNodeName} has no edge for "${targetBranchName}".`)
+                unlockFailure(`${requestedBotName}:${requestedBranchName} has no edge for "${targetBranchName}".`)
                 $('#result').append(renderSingleBranchTable(requestedNode))
                 return
             }
