Fix track changes when dealing with isolated files and files that have been branched from another file in the same stream

diff --git a/Engine/Extras/RoboMerge/v3/src/robo/trackchange.ts b/Engine/Extras/RoboMerge/v3/src/robo/trackchange.ts
--- a/Engine/Extras/RoboMerge/v3/src/robo/trackchange.ts
+++ b/Engine/Extras/RoboMerge/v3/src/robo/trackchange.ts
@@ -179,15 +179,19 @@ export async function trackChange(robo: RoboMerge, logger: ContextualLogger, ini
                         // integrated for move/delete will point at the paired move/add not where it was merged to in another stream, so not useful to evaluate it
                         if (entry.action != 'move/delete') {
                             // Do an integrated check
-                            const integrated = (await p4.integrated(null, entry.depotFile)).filter(x => x.how.includes('from') && x.toFile == entry.depotFile && x.endToRev == `#${entry.rev}`)
+                            const integrated = (await p4.integrated(null, entry.depotFile))
+                                                        .filter(   x => x.how.includes('from') 
+                                                                && x.toFile == entry.depotFile 
+                                                                && x.endToRev == `#${entry.rev}` 
+                                                                && getStreamFromPath(x.toFile) != getStreamFromPath(x.fromFile))
                             if (integrated.length > 0) {
                                 const fstat = await p4.fstat(null, `${integrated[0].fromFile}${integrated[0].endFromRev}`)
                                 opts = { ...opts, parentCL: {cl: fstat[0].headChange, serverId: p4.serverID } }
                                 sourceChain.push(opts.parentCL!)
+                                break
                             }
-                            break
                         }
-                        else if (desc.entries.length == 1) {
+                        if (desc.entries.length == 1) {
                             await getMoreDescEntries(p4, cl, desc)
                         }
                     }
@@ -325,12 +329,12 @@ export async function trackChange(robo: RoboMerge, logger: ContextualLogger, ini
                 if (runIntegrated) {
                     const integrated = await p4.integrated(null, entry.depotFile, {intoOnly: true, startCL: clToConsider.cl})
                     if (integrated.length > 0) {
-                        runIntegrated = false
                         for (let integ of integrated) {
                             const startToRev = integ.startToRev == "#none" ? 0 : parseInt(integ.startToRev.slice(1))
                             const endToRev = parseInt(integ.endToRev.slice(1))
                             if (entry.rev <= endToRev && entry.rev > startToRev)
                             {
+                                runIntegrated = false
                                 const destChange = changes.get(integ.change)
                                 if (destChange) {
                                     destChange.sourceCL = clToConsider
