Fix trackchanges when source nodes are not robomerge nodes and there was a manual merge

diff --git a/Engine/Extras/RoboMerge/v3/src/robo/trackchange.ts b/Engine/Extras/RoboMerge/v3/src/robo/trackchange.ts
--- a/Engine/Extras/RoboMerge/v3/src/robo/trackchange.ts
+++ b/Engine/Extras/RoboMerge/v3/src/robo/trackchange.ts
@@ -44,6 +44,7 @@ interface GatherOptions {
     parentCL?: TrackedCL,
     lastCL?: TrackedCL,
     initialSubmit?: boolean
+    inSourceChain?: boolean
 }
 
 export async function trackChange(robo: RoboMerge, logger: ContextualLogger, initialCL: TrackedCL, userTags: Set<string>, filters: TrackFilters) {
@@ -125,7 +126,23 @@ export async function trackChange(robo: RoboMerge, logger: ContextualLogger, ini
         }
     }
 
-    let changes = new Map<string, TrackedChange>()
+    class Changes {
+        private changes = new Map<string, TrackedChange>()
+
+        private static changeKey(cl: TrackedCL) {
+            return JSON.stringify({cl: `${cl.cl}`, serverId: cl.serverId || 'default'})
+        }
+
+        public get(cl: TrackedCL) {
+            return this.changes.get(Changes.changeKey(cl))
+        }
+
+        public set(cl: TrackedCL, tc: TrackedChange) {
+            this.changes.set(Changes.changeKey(cl), tc)
+
+        }
+    }
+    const changes = new Changes
 
     let determineSourceServer = (node: Node) => {
         let edges = graph.getEdgesForNode(node)
@@ -143,7 +160,7 @@ export async function trackChange(robo: RoboMerge, logger: ContextualLogger, ini
     let gatherCLInfo = async (cl: TrackedCL, opts?: GatherOptions) => {
 
         if (cl.serverId === undefined && opts?.lastCL) {
-            const lastChange = changes.get(JSON.stringify(opts.lastCL))
+            const lastChange = changes.get(opts.lastCL)
             if (lastChange?.node) {
                 cl.serverId = determineSourceServer(lastChange.node)
             }
@@ -199,8 +216,8 @@ export async function trackChange(robo: RoboMerge, logger: ContextualLogger, ini
                 initialSubmit = opts?.parentCL === undefined
             }
             const mergeMethod = getMergeMethod(desc.description, cl, {...opts, initialSubmit} )
-            if (clNode || (RobomergeMethodStrings as readonly string[]).includes(mergeMethod)) {
-                changes.set(JSON.stringify(cl), {desc, mergeMethod, node: clNode, sourceCL: opts?.parentCL, destCLs: (opts && opts.lastCL ? [opts.lastCL] : []) })
+            if (clNode || opts?.inSourceChain || (RobomergeMethodStrings as readonly string[]).includes(mergeMethod)) {
+                changes.set(cl, {desc, mergeMethod, node: clNode, sourceCL: opts?.parentCL, destCLs: (opts && opts.lastCL ? [opts.lastCL] : []) })
             }
         } 
         catch (reason) {
@@ -227,9 +244,9 @@ export async function trackChange(robo: RoboMerge, logger: ContextualLogger, ini
     for (let i=0; i < sourceChain.length; i++) {
         const sourceCL = sourceChain[i]
         const parentCL = sourceChain[i+1]
-        if (await gatherCLInfo(sourceCL, {parentCL, lastCL})) {
+        if (await gatherCLInfo(sourceCL, {parentCL, lastCL, inSourceChain: true})) {
             if (i == 0) {
-                let initialChange = changes.get(JSON.stringify(initialCL))!
+                let initialChange = changes.get(initialCL)!
                 // Fix up the initial changes merge method and source now that we know it is not the initial submit
                 initialChange.mergeMethod = getMergeMethod(initialChange.desc.description, initialCL, {parentCL: sourceCL})
                 initialChange.sourceCL = sourceCL
@@ -247,7 +264,7 @@ export async function trackChange(robo: RoboMerge, logger: ContextualLogger, ini
     let clsToConsider: TrackedCL[] = [data.originalCL]
 
     const considerCL = async (clToConsider: TrackedCL) => {
-        let changeToConsider = changes.get(JSON.stringify(clToConsider))
+        let changeToConsider = changes.get(clToConsider)
         if (!changeToConsider) {
             return
         }
@@ -335,11 +352,11 @@ export async function trackChange(robo: RoboMerge, logger: ContextualLogger, ini
                             if (entry.rev <= endToRev && entry.rev > startToRev)
                             {
                                 runIntegrated = false
-                                const destChange = changes.get(integ.change)
+                                const newChangeToConsider = {cl: integ.change, serverId: clToConsider.serverId}
+                                const destChange = changes.get(newChangeToConsider)
                                 if (destChange) {
                                     destChange.sourceCL = clToConsider
                                 } else {
-                                    const newChangeToConsider = {cl: integ.change, serverId: clToConsider.serverId}
                                     changeToConsider.destCLs.push(newChangeToConsider)
                                     await gatherCLInfo(newChangeToConsider, {parentCL: clToConsider})
                                 }
