Fix trackchanges failing under certain cases

diff --git a/Engine/Extras/RoboMerge/v3/src/robo/trackchange.ts b/Engine/Extras/RoboMerge/v3/src/robo/trackchange.ts
--- a/Engine/Extras/RoboMerge/v3/src/robo/trackchange.ts
+++ b/Engine/Extras/RoboMerge/v3/src/robo/trackchange.ts
@@ -129,7 +129,7 @@ export async function trackChange(robo: RoboMerge, logger: ContextualLogger, ini
     class Changes {
         private changes = new Map<string, TrackedChange>()
 
-        private static changeKey(cl: TrackedCL) {
+        static changeKey(cl: TrackedCL) {
             return JSON.stringify({cl: `${cl.cl}`, serverId: cl.serverId || 'default'})
         }
 
@@ -355,6 +355,10 @@ export async function trackChange(robo: RoboMerge, logger: ContextualLogger, ini
                                 const newChangeToConsider = {cl: integ.change, serverId: clToConsider.serverId}
                                 const destChange = changes.get(newChangeToConsider)
                                 if (destChange) {
+                                    const newChangeKey = Changes.changeKey(newChangeToConsider)
+                                    if (!changeToConsider.destCLs.some(cl => Changes.changeKey(cl) === newChangeKey)) {
+                                        changeToConsider.destCLs.push(newChangeToConsider)
+                                    }
                                     destChange.sourceCL = clToConsider
                                 } else {
                                     changeToConsider.destCLs.push(newChangeToConsider)
