Allow certain edges to not be displayed by default

diff --git a/Engine/Extras/RoboMerge/v3/README.md b/Engine/Extras/RoboMerge/v3/README.md
--- a/Engine/Extras/RoboMerge/v3/README.md
+++ b/Engine/Extras/RoboMerge/v3/README.md
@@ -102,6 +102,7 @@ for both node and edge.
 |`forcePause`               |ne|            | |If flag set, pause - applies each restart| |
 |`additionalSlackChannelForBlockages` |e |  | |Single extra Slack channel               | |
 |`blockAssetFlow`           |n |`[]`        | |Reject any integrations containing assets to these targets   | |
+|`graphFilteredEdges`       |n |`[]`        | |By default these edges will not be displayed on the graph  | |
 |`integrationMethod`        |e |normal      | |For skip if unchanged                    | |
 |`disallowSkip`             |e |            | |Remove skip option from UI               | |
 |`emailOnBlockage`          |e |`true`      | |Email owners of conflicts?               | |
diff --git a/Engine/Extras/RoboMerge/v3/public/allbots.html b/Engine/Extras/RoboMerge/v3/public/allbots.html
--- a/Engine/Extras/RoboMerge/v3/public/allbots.html
+++ b/Engine/Extras/RoboMerge/v3/public/allbots.html
@@ -93,6 +93,12 @@
 							<label for="hideDisconnected" title="Hide nodes with no edges">Hide Disconnected</label>	
 						</div>
 					</div>
+					<div class="row form-check">
+						<div class="col">
+							<input class="form-check-input" type="checkbox" id="showFiltered" name="showFiltered" value="showFiltered">
+							<label  for="showFiltered"  title="Show edges that have been flagged as filtered">Show Filtered</label>	
+						</div>
+					</div>
 					<div class="row form-check">
 						<div class="col">
 							<input class="form-check-input" type="checkbox" id="noGroups" name="noGroups" value="noGroups">
diff --git a/Engine/Extras/RoboMerge/v3/public/js/allbots.js b/Engine/Extras/RoboMerge/v3/public/js/allbots.js
--- a/Engine/Extras/RoboMerge/v3/public/js/allbots.js
+++ b/Engine/Extras/RoboMerge/v3/public/js/allbots.js
@@ -33,6 +33,10 @@ function populateControlPanel(bots, allBots)
     hideDisconnected.checked = option.hideDisconnected;
     hideDisconnected.addEventListener('change', controlPanelChanged);
 
+    let showFiltered = document.getElementById('showFiltered');
+    showFiltered.checked = option.showFiltered;
+    showFiltered.addEventListener('change', controlPanelChanged);
+
     let noGroups = document.getElementById('noGroups');
     noGroups.checked = option.noGroups;
     noGroups.addEventListener('change', controlPanelChanged);
@@ -54,6 +58,9 @@ function controlPanelChanged()
     const hideDisconnected = document.getElementById('hideDisconnected');
 	options.hideDisconnected = hideDisconnected.checked;
 
+    const showFiltered = document.getElementById('showFiltered');
+    options.showFiltered = showFiltered.checked;
+
     const noGroups = document.getElementById('noGroups');
     options.noGroups = noGroups.checked;
 
@@ -90,6 +97,12 @@ function createUrlFromParameters(reload)
         args.push("hideDisconnected");
     }
 
+    const showFiltered = document.getElementById('showFiltered');
+    if (showFiltered.checked)
+    {
+        args.push("showFiltered");
+    }
+
     const noGroups = document.getElementById('noGroups');
     if (noGroups.checked)
     {
diff --git a/Engine/Extras/RoboMerge/v3/public/js/flow.graph.js b/Engine/Extras/RoboMerge/v3/public/js/flow.graph.js
--- a/Engine/Extras/RoboMerge/v3/public/js/flow.graph.js
+++ b/Engine/Extras/RoboMerge/v3/public/js/flow.graph.js
@@ -6,6 +6,7 @@ class Info {
         this.tooltip = tooltip;
         this.name = name;
         this.sources = sources;
+        this.filteredEdges = new Set();
         this.forcedEdges = new Set();
         this.defaultDests = new Set();
         this.blockAssetDests = new Set();
@@ -14,6 +15,7 @@ class Info {
     update(branch, getId) {
         const edgeId = (alias) => branch.def.name + '->' + getId(alias);
         // shuld really do the same for default and block assets
+        this.filteredEdges = new Set([...this.filteredEdges, ...branch.def.graphFilteredTargets.map(edgeId)]);
         this.forcedEdges = new Set([...this.forcedEdges, ...branch.def.forceFlowTo.map(edgeId)]);
         this.defaultDests = new Set([...this.defaultDests, ...branch.def.defaultFlow.map(edgeId)]);
         this.blockAssetDests = new Set([...this.blockAssetDests, ...branch.def.blockAssetTargets.map(edgeId)]);
@@ -198,6 +200,10 @@ class Graph {
         if (!isForced && this.options.showOnlyForced) {
             return;
         }
+        const isFiltered = hasEdge(srcInfo.filteredEdges, dstInfo);
+        if (isFiltered && !this.options.showFiltered) {
+            return;
+        }
         const isGated = new Map(srcBranchStatus.def.edgeProperties).get(dst).lastGoodCLPath
         this.connectedNodes.add(srcInfo.id);
         this.connectedNodes.add(dstInfo.id);
diff --git a/Engine/Extras/RoboMerge/v3/public/js/flow.js b/Engine/Extras/RoboMerge/v3/public/js/flow.js
--- a/Engine/Extras/RoboMerge/v3/public/js/flow.js
+++ b/Engine/Extras/RoboMerge/v3/public/js/flow.js
@@ -2,6 +2,7 @@
 
 class FlowOptions{
     showOnlyForced = false;
+    showFiltered = false;
     noGroups = false;
     hideDisconnected = false;
 }
@@ -33,6 +34,7 @@ function parseOptions(search) {
     let options = new FlowOptions();
     
     options.showOnlyForced = !!search.match(/showOnlyForced/i);
+    options.showFiltered = !!search.match(/showFiltered/i);
     options.noGroups = !!search.match(/noGroups/i);
     options.hideDisconnected = !!search.match(/hideDisconnected/i);
 
diff --git a/Engine/Extras/RoboMerge/v3/public/js/preview.js b/Engine/Extras/RoboMerge/v3/public/js/preview.js
--- a/Engine/Extras/RoboMerge/v3/public/js/preview.js
+++ b/Engine/Extras/RoboMerge/v3/public/js/preview.js
@@ -34,6 +34,9 @@ function controlPanelChanged()
     const hideDisconnected = document.getElementById('hideDisconnected');
 	options.hideDisconnected = hideDisconnected.checked;
 
+    const showFiltered = document.getElementById('showFiltered');
+    options.showFiltered = showFiltered.checked;
+
     const noGroups = document.getElementById('noGroups');
     options.noGroups = noGroups.checked;
 
@@ -66,6 +69,10 @@ function populateControlPanel(bots, allBots)
 	hideDisconnected.checked = option.hideDisconnected;
 	hideDisconnected.addEventListener('change', controlPanelChanged);
 
+    let showFiltered = document.getElementById('showFiltered');
+	showFiltered.checked = option.showFiltered;
+	showFiltered.addEventListener('change', controlPanelChanged);
+
     let noGroups = document.getElementById('noGroups');
 	noGroups.checked = option.noGroups;
 	noGroups.addEventListener('change', controlPanelChanged);
@@ -94,6 +101,12 @@ function createUrlFromParameters(reload)
         args.push("hideDisconnected");
     }
 
+    const showFiltered = document.getElementById('showFiltered');
+    if (showFiltered.checked)
+    {
+        args.push("showFiltered");
+    }
+
     const noGroups = document.getElementById('noGroups');
     if (noGroups.checked)
     {
diff --git a/Engine/Extras/RoboMerge/v3/public/preview.html b/Engine/Extras/RoboMerge/v3/public/preview.html
--- a/Engine/Extras/RoboMerge/v3/public/preview.html
+++ b/Engine/Extras/RoboMerge/v3/public/preview.html
@@ -79,6 +79,12 @@
 							<label for="hideDisconnected" title="Hide nodes with no edges">Hide Disconnected</label>	
 						</div>
 					</div>
+					<div class="row form-check">
+						<div class="col">
+							<input class="form-check-input" type="checkbox" id="showFiltered" name="showFiltered" value="showFiltered">
+							<label  for="showFiltered"  title="Show edges that have been flagged as filtered">Show Filtered</label>	
+						</div>
+					</div>
 					<div class="row form-check">
 						<div class="col">
 							<input class="form-check-input" type="checkbox" id="noGroups" name="noGroups" value="noGroups">
diff --git a/Engine/Extras/RoboMerge/v3/src/robo/branch-interfaces.ts b/Engine/Extras/RoboMerge/v3/src/robo/branch-interfaces.ts
--- a/Engine/Extras/RoboMerge/v3/src/robo/branch-interfaces.ts
+++ b/Engine/Extras/RoboMerge/v3/src/robo/branch-interfaces.ts
@@ -103,6 +103,7 @@ export type EditableBranch = BranchBase & {
 	enabled: boolean
 	visibility: string[] | string
 	blockAssetTargets: Set<string>
+	graphFilteredTargets: Set<string>
 	allowDeadend: boolean
 
 	edgeProperties: Map<string, EdgeOptions>
diff --git a/Engine/Extras/RoboMerge/v3/src/robo/branchdefs.ts b/Engine/Extras/RoboMerge/v3/src/robo/branchdefs.ts
--- a/Engine/Extras/RoboMerge/v3/src/robo/branchdefs.ts
+++ b/Engine/Extras/RoboMerge/v3/src/robo/branchdefs.ts
@@ -115,6 +115,7 @@ const nodeOptionFieldsPrototype = {
 	forceAll: false,
 	visibility: '' as string[] | string,
 	blockAssetFlow: [''],
+	graphFilteredEdges: [''],
 	disallowDeadend: false,
 
 	streamServer: undefined as string | undefined,
diff --git a/Engine/Extras/RoboMerge/v3/src/robo/branchgraph.ts b/Engine/Extras/RoboMerge/v3/src/robo/branchgraph.ts
--- a/Engine/Extras/RoboMerge/v3/src/robo/branchgraph.ts
+++ b/Engine/Extras/RoboMerge/v3/src/robo/branchgraph.ts
@@ -239,6 +239,7 @@ if (botname === '__TEST__') {
 			defaultFlow: options.defaultFlow || [],
 			macros: options.macros,
 			blockAssetTargets: new Set<string>(),
+			graphFilteredTargets: new Set<string>(),
 			enabled: !options.disabled,
 			allowDeadend: !options.disallowDeadend,
 			changelistParser: BranchGraph._getConfig(options.changelistParser, this.config.changelistParser) || null,
@@ -291,6 +292,17 @@ if (botname === '__TEST__') {
 			}
 		}
 
+		if (options.graphFilteredEdges) {
+			for (const filtered of options.graphFilteredEdges) {
+				const flowUpper = filtered.toUpperCase()
+				if (branch.flowsTo.indexOf(flowUpper) === -1) {
+					throw new Error(`Graph filtered edge '${filtered}' not in flowsTo list`)
+				}
+				branch.graphFilteredTargets.add(flowUpper)
+
+			}
+		}
+
 		for (let i=0; i<branch.defaultFlow.length; ++i) {
 			branch.defaultFlow[i] = branch.defaultFlow[i].toUpperCase()
 		}
diff --git a/Engine/Extras/RoboMerge/v3/src/robo/status-types.ts b/Engine/Extras/RoboMerge/v3/src/robo/status-types.ts
--- a/Engine/Extras/RoboMerge/v3/src/robo/status-types.ts
+++ b/Engine/Extras/RoboMerge/v3/src/robo/status-types.ts
@@ -153,6 +153,7 @@ export type BranchDefForStatus = {
 	bot: string
 	config: any //NodeOptions
 	defaultFlow: string[]
+	graphFilteredTargets: string[]
 	flowsTo: string[]
 	forceFlowTo: string[]
 }
