Default to not shwing groups

diff --git a/Engine/Extras/RoboMerge/v3/public/allbots.html b/Engine/Extras/RoboMerge/v3/public/allbots.html
--- a/Engine/Extras/RoboMerge/v3/public/allbots.html
+++ b/Engine/Extras/RoboMerge/v3/public/allbots.html
@@ -101,8 +101,8 @@
 					</div>
 					<div class="row form-check">
 						<div class="col">
-							<input class="form-check-input" type="checkbox" id="noGroups" name="noGroups" value="noGroups">
-							<label  for="noGroups"  title="Relevant if showing multiple bots: suppress boxes grouping each bot">No Groups</label>	
+							<input class="form-check-input" type="checkbox" id="showGroups" name="showGroups" value="showGroups">
+							<label  for="showGroups"  title="Relevant if showing multiple bots: show boxes grouping each bot">Show Groups</label>	
 						</div>
 					</div>
 					<div class="row form-check">
diff --git a/Engine/Extras/RoboMerge/v3/public/js/allbots.js b/Engine/Extras/RoboMerge/v3/public/js/allbots.js
--- a/Engine/Extras/RoboMerge/v3/public/js/allbots.js
+++ b/Engine/Extras/RoboMerge/v3/public/js/allbots.js
@@ -37,9 +37,9 @@ function populateControlPanel(bots, allBots)
     showFiltered.checked = option.showFiltered;
     showFiltered.addEventListener('change', controlPanelChanged);
 
-    let noGroups = document.getElementById('noGroups');
-    noGroups.checked = option.noGroups;
-    noGroups.addEventListener('change', controlPanelChanged);
+    let showGroups = document.getElementById('showGroups');
+    showGroups.checked = option.showGroups;
+    showGroups.addEventListener('change', controlPanelChanged);
     
     let showOnlyForced = document.getElementById('showOnlyForced');
     showOnlyForced.checked = option.showOnlyForced;
@@ -61,8 +61,8 @@ function controlPanelChanged()
     const showFiltered = document.getElementById('showFiltered');
     options.showFiltered = showFiltered.checked;
 
-    const noGroups = document.getElementById('noGroups');
-    options.noGroups = noGroups.checked;
+    const showGroups = document.getElementById('showGroups');
+    options.showGroups = showGroups.checked;
 
     const showOnlyForced = document.getElementById('showOnlyForced');
 	options.showOnlyForced = showOnlyForced.checked;
@@ -103,10 +103,10 @@ function createUrlFromParameters(reload)
         args.push("showFiltered");
     }
 
-    const noGroups = document.getElementById('noGroups');
-    if (noGroups.checked)
+    const showGroups = document.getElementById('showGroups');
+    if (showGroups.checked)
     {
-        args.push("noGroups");
+        args.push("showGroups");
     }
 
     const showOnlyForced = document.getElementById('showOnlyForced');
diff --git a/Engine/Extras/RoboMerge/v3/public/js/flow.graph.js b/Engine/Extras/RoboMerge/v3/public/js/flow.graph.js
--- a/Engine/Extras/RoboMerge/v3/public/js/flow.graph.js
+++ b/Engine/Extras/RoboMerge/v3/public/js/flow.graph.js
@@ -64,7 +64,7 @@ class Graph {
             ? this.allBranches.filter(b => this.options.botsToShow.indexOf(b.bot) >= 0)
             : this.allBranches;
         this.findSharedNodes();
-        this.addNodesAndEdges(!this.options.noGroups);
+        this.addNodesAndEdges(this.options.showGroups);
 
         let title = 'All bots';
         if (this.options.botsToShow)
diff --git a/Engine/Extras/RoboMerge/v3/public/js/flow.js b/Engine/Extras/RoboMerge/v3/public/js/flow.js
--- a/Engine/Extras/RoboMerge/v3/public/js/flow.js
+++ b/Engine/Extras/RoboMerge/v3/public/js/flow.js
@@ -3,7 +3,7 @@
 class FlowOptions{
     showOnlyForced = false;
     showFiltered = false;
-    noGroups = false;
+    showGroups = false;
     hideDisconnected = false;
 }
 
@@ -35,7 +35,7 @@ function parseOptions(search) {
     
     options.showOnlyForced = !!search.match(/showOnlyForced/i);
     options.showFiltered = !!search.match(/showFiltered/i);
-    options.noGroups = !!search.match(/noGroups/i);
+    options.showGroups = !!search.match(/showGroups/i);
     options.hideDisconnected = !!search.match(/hideDisconnected/i);
 
     return options;
diff --git a/Engine/Extras/RoboMerge/v3/public/js/preview.js b/Engine/Extras/RoboMerge/v3/public/js/preview.js
--- a/Engine/Extras/RoboMerge/v3/public/js/preview.js
+++ b/Engine/Extras/RoboMerge/v3/public/js/preview.js
@@ -37,8 +37,8 @@ function controlPanelChanged()
     const showFiltered = document.getElementById('showFiltered');
     options.showFiltered = showFiltered.checked;
 
-    const noGroups = document.getElementById('noGroups');
-    options.noGroups = noGroups.checked;
+    const showGroups = document.getElementById('showGroups');
+    options.showGroups = showGroups.checked;
 
     const showOnlyForced = document.getElementById('showOnlyForced');
 	options.showOnlyForced = showOnlyForced.checked;
@@ -73,9 +73,9 @@ function populateControlPanel(bots, allBots)
 	showFiltered.checked = option.showFiltered;
 	showFiltered.addEventListener('change', controlPanelChanged);
 
-    let noGroups = document.getElementById('noGroups');
-	noGroups.checked = option.noGroups;
-	noGroups.addEventListener('change', controlPanelChanged);
+    let showGroups = document.getElementById('showGroups');
+	showGroups.checked = option.showGroups;
+	showGroups.addEventListener('change', controlPanelChanged);
     
     let showOnlyForced = document.getElementById('showOnlyForced');
 	showOnlyForced.checked = option.showOnlyForced;
@@ -107,10 +107,10 @@ function createUrlFromParameters(reload)
         args.push("showFiltered");
     }
 
-    const noGroups = document.getElementById('noGroups');
-    if (noGroups.checked)
+    const showGroups = document.getElementById('showGroups');
+    if (showGroups.checked)
     {
-        args.push("noGroups");
+        args.push("showGroups");
     }
 
     const showOnlyForced = document.getElementById('showOnlyForced');
diff --git a/Engine/Extras/RoboMerge/v3/public/preview.html b/Engine/Extras/RoboMerge/v3/public/preview.html
--- a/Engine/Extras/RoboMerge/v3/public/preview.html
+++ b/Engine/Extras/RoboMerge/v3/public/preview.html
@@ -87,8 +87,8 @@
 					</div>
 					<div class="row form-check">
 						<div class="col">
-							<input class="form-check-input" type="checkbox" id="noGroups" name="noGroups" value="noGroups">
-							<label  for="noGroups"  title="Relevant if showing multiple bots: suppress boxes grouping each bot">No Groups</label>	
+							<input class="form-check-input" type="checkbox" id="showGroups" name="showGroups" value="showGroups">
+							<label  for="showGroups"  title="Relevant if showing multiple bots: show boxes grouping each bot">Show Groups</label>	
 						</div>
 					</div>
 					<div class="row form-check">
