Prevent gates from getting queued up and causing it to be stuck on current gate

diff --git a/Engine/Extras/RoboMerge/v3/src/robo/gate.ts b/Engine/Extras/RoboMerge/v3/src/robo/gate.ts
--- a/Engine/Extras/RoboMerge/v3/src/robo/gate.ts
+++ b/Engine/Extras/RoboMerge/v3/src/robo/gate.ts
@@ -238,7 +238,7 @@ export class Gate {
 			}
 
 			// if no current and in window, kick off new
-			if (!this.currentGateInfo && this.queuedGates.length > 0) {
+			if (this.queuedGates.length > 0 && (!this.currentGateInfo || this.currentGateInfo.cl <= this.lastCl) ) {
 				if (this.tryNextGate(this.lastCl)) {
 					this.reportCatchingUp()
 					dirty = true
