Fix last CL not displaying if the server doesn't have swarm

diff --git a/Engine/Extras/RoboMerge/v3/public/js/boilerplate.js b/Engine/Extras/RoboMerge/v3/public/js/boilerplate.js
--- a/Engine/Extras/RoboMerge/v3/public/js/boilerplate.js
+++ b/Engine/Extras/RoboMerge/v3/public/js/boilerplate.js
@@ -88,9 +88,7 @@ function makeClLink(cl, swarmURL, alias) {
 	if (swarmURL) {
 		return `<a href="${swarmURL}/changes/${cl}" target="_blank">${alias || cl}</a>`
 	}
-	else {
-		return `${alias || cl}`
-	}
+	return `<span>${alias || cl}</span>`
 }
 
 function makeSwarmFileLink(depotPath, swarmURL, alias) {
