Add mechanism to save the bot graph to a PNG

diff --git a/Engine/Extras/RoboMerge/v3/public/js/flow.js b/Engine/Extras/RoboMerge/v3/public/js/flow.js
--- a/Engine/Extras/RoboMerge/v3/public/js/flow.js
+++ b/Engine/Extras/RoboMerge/v3/public/js/flow.js
@@ -65,6 +65,28 @@ function getAllBots(data) {
     return allBots;
 }
 
+function savePNG() {
+    const svg = document.getElementById("svgSpan").querySelector('svg');
+    const svgBoundingRect = svg.getBoundingClientRect();
+
+    var canvas = document.getElementById("cnvs");
+    canvas.width = svgBoundingRect.width*1.02; // 2% wider seems to be needed to actually fit the image on the canvas
+    canvas.height = svgBoundingRect.height;
+
+    const svgData = new XMLSerializer().serializeToString(svg);
+    const imgBlob = new Blob([svgData], { type: 'image/svg+xml' });
+    var img = new Image();
+    img.onload = function () {
+        canvas.getContext('2d').drawImage(img, 0, 0);
+        var a = document.createElement('a');
+        a.href = canvas.toDataURL('image/png').replace('image/png', 'octet/stream');
+        a.download = 'botGraph.png'; // TODO: Dynamic Name Generation
+        a.click();
+    };
+    const win = window.URL || window.webkitURL || window;
+    img.src = win.createObjectURL(imgBlob);
+}
+
 function showFlowGraph(data, args, addlegend) {
     const lines = makeGraph(data, args);
     const graphContainer = $('<div class="clearfix">');
@@ -84,13 +106,15 @@ function showFlowGraph(data, args, addlegend) {
         
 
         $('#graph-loading-text').hide();
-        const span = $('<span>').css('margin', 'auto').html(svg);
+        const span = $('<span id="svgSpan">').css('margin', 'auto').html(svg);
         const svgEl = $('svg', span).addClass('branch-graph').removeAttr('width');
         // scale graph to 70% of default size
         const height = Math.round(parseInt(svgEl.attr('height')) * .7);
         svgEl.attr('height', height + 'px').css('vertical-align', 'top');
         flowGraph.empty();
+        flowGraph.append($('<img src="/img/save.png" style="position:absolute" onclick="savePNG()" title="Save Graph as PNG" width="16">'));
         flowGraph.append(span);
+        flowGraph.append($('<canvas id="cnvs" class="hidden"></canvas>'))
     });
     return graphContainer;
 }
