Improve error outputting

diff --git a/Engine/Extras/RoboMerge/v3/src/robo/robo.ts b/Engine/Extras/RoboMerge/v3/src/robo/robo.ts
--- a/Engine/Extras/RoboMerge/v3/src/robo/robo.ts
+++ b/Engine/Extras/RoboMerge/v3/src/robo/robo.ts
@@ -508,7 +508,7 @@ if (!args.noIPC) {
 	})
 }
 else {
-	const sendMessage = (name: string, args?: any[]) => ipc.handle({name:name, args:args}).catch(err => console.log('IPC error! ' + err.toString()));
+	const sendMessage = (name: string, args?: any[]) => ipc.handle({name:name, args:args}).catch(err => console.log(`IPC error! ${err.toString()}\n${err.stack}`));
 
 	const ws = new RoboServer(args.externalUrl, sendMessage, () => roboStartupLogger.getLogTail(),
 		() => roboStartupLogger.info('Received: getLastCrash'),
