Improve stomp
* Stomp verification no longer requires retrying the merge
* Stomping does not do a merge, verify, revert, and then again do the merge

diff --git a/Engine/Extras/RoboMerge/v3/src/common/perforce.ts b/Engine/Extras/RoboMerge/v3/src/common/perforce.ts
--- a/Engine/Extras/RoboMerge/v3/src/common/perforce.ts
+++ b/Engine/Extras/RoboMerge/v3/src/common/perforce.ts
@@ -424,6 +424,10 @@ export class ResolveResult {
 		}
 	}
 
+	getHasUnknownConflicts() {
+		return this.hasUnknownConflicts		
+	}
+	
 	hasConflict() {
 		return this.remainingConflicts.length > 0 || this.hasUnknownConflicts
 	}
diff --git a/Engine/Extras/RoboMerge/v3/src/robo/branch-interfaces.ts b/Engine/Extras/RoboMerge/v3/src/robo/branch-interfaces.ts
--- a/Engine/Extras/RoboMerge/v3/src/robo/branch-interfaces.ts
+++ b/Engine/Extras/RoboMerge/v3/src/robo/branch-interfaces.ts
@@ -1,6 +1,6 @@
 // Copyright Epic Games, Inc. All Rights Reserved.
 
-import { BranchSpec, ConflictedResolveNFile, EdgeServer, ExclusiveFileDetails } from '../common/perforce';
+import { BranchSpec, ConflictedResolveNFile, DescribeResult, EdgeServer, ExclusiveFileDetails, ResolveResult } from '../common/perforce';
 // TODO: Remove Circular Dependency on bot-interfaces
 import { NodeBotInterface } from './bot-interfaces';
 import { FailureKind } from './status-types';
@@ -59,6 +59,23 @@ export interface StompVerification extends OperationResult {
 	remainingAllBinary?: boolean // Check to see if unresolved non-binary files are remaining -- we shouldn't stomp those!
 }
 
+export function GetStompVerificationResult(verifyResult: StompVerification): OperationResult {
+	if (!verifyResult.success) {
+		return { success: false, message: `stompChanges Request Failed Due to Bad Verification: ${verifyResult.message}` }
+	}
+
+	if (!verifyResult.validRequest) {
+		return { success: false, message: `stompChanges Request Cannot Be Completed: ${verifyResult.message}` }
+	}
+
+	// One last check to ensure this is an all binary stomp
+	if (!verifyResult.remainingAllBinary) {
+		return { success: false, message: `stompChanges Failed -- Verification had remaining text files (${verifyResult.message})` }
+	}
+
+	return { success: true, message: 'Queued stomp request' }
+}
+
 export interface ExclusiveLockUser {
     user: string
     userEmail: Promise<string | null>
@@ -204,6 +221,39 @@ export interface AlreadyIntegrated {
 	action: MergeAction
 }
 
+export class ResolveResultAdditionalInfo {
+
+	protected remainingConflicts: ConflictedResolveNFile[]
+	readonly conflictsText: string
+	readonly hasUnknownConflicts: boolean
+	readonly parsingError: string = ""
+	readonly successfullyParsedConflicts: boolean
+
+	constructor(resolveResult?: ResolveResult) {
+		if (resolveResult) {
+			this.conflictsText = resolveResult.getConflictsText()
+			this.hasUnknownConflicts = resolveResult.getHasUnknownConflicts()
+			this.parsingError = resolveResult.getParsingError()
+			this.remainingConflicts = resolveResult.getConflicts()
+			this.successfullyParsedConflicts = resolveResult.successfullyParsedConflicts()
+		}
+	}
+
+	hasConflict() {
+		return this.remainingConflicts.length > 0 || this.hasUnknownConflicts
+	}
+
+	getConflicts() {
+		return this.remainingConflicts
+	}
+}
+
+export interface MergeConflictAdditionalInformation {
+	resolveResult: ResolveResultAdditionalInfo,
+	describeResult: DescribeResult,
+	client?: string
+}
+
 export interface ForcedCl {
 	nodeOrEdgeName: string
 	sourceBranchUpperName: string
diff --git a/Engine/Extras/RoboMerge/v3/src/robo/edgebot.ts b/Engine/Extras/RoboMerge/v3/src/robo/edgebot.ts
--- a/Engine/Extras/RoboMerge/v3/src/robo/edgebot.ts
+++ b/Engine/Extras/RoboMerge/v3/src/robo/edgebot.ts
@@ -5,10 +5,11 @@ import * as p4util from '../common/p4util';
 
 import { ContextualLogger } from "../common/logger";
 import { Recipients } from "../common/mailer";
-import { Change, coercePerforceWorkspace, ConflictedResolveNFile, EdgeServer, EditChangeOpts, ExclusiveFileDetails, EXCLUSIVE_CHECKOUT_REGEX, getRootDirectoryForBranch, IntegrationSource, IntegrationTarget, isExecP4Error, PerforceContext } from "../common/perforce";
+import { Change, coercePerforceWorkspace, ConflictedResolveNFile, EdgeServer, EditChangeOpts, ExclusiveFileDetails, EXCLUSIVE_CHECKOUT_REGEX } from "../common/perforce";
+import { getRootDirectoryForBranch, IntegrationSource, IntegrationTarget, isExecP4Error, PerforceContext } from "../common/perforce";
 import { VersionReader } from "../common/version";
 import { EdgeBotInterface, IPCControls, ReconsiderArgs } from "./bot-interfaces";
-import { AlreadyIntegrated, Branch, ChangeInfo, ConflictingFile, Failure, MergeAction, PendingChange } from "./branch-interfaces";
+import { AlreadyIntegrated, Branch, ChangeInfo, ConflictingFile, Failure, GetStompVerificationResult, MergeAction, MergeConflictAdditionalInformation, PendingChange, ResolveResultAdditionalInfo } from "./branch-interfaces";
 import { EdgeOptions, IntegrationMethod } from "./branchdefs";
 import { PersistentConflict } from "./conflict-interfaces";
 import { BotEventTriggers } from "./events";
@@ -741,13 +742,37 @@ class EdgeBotImpl extends PerforceStatefulBot {
 
 		const edgeServerAddress = pending.change.edgeServerToHostShelf && pending.change.edgeServerToHostShelf.address
 
-		const result = await this.p4.resolve(
+		let result = await this.p4.resolve(
 			pending.change.targetWorkspace,
 			pending.newCl,
-			pending.action.mergeMode,
+			pending.change.forceStompChanges ? 'normal' : pending.action.mergeMode,
 			false, // detail === 'quick',
 			edgeServerAddress
+		)
+
+		if (pending.change.forceStompChanges && result.hasConflict()) {
+			const additionalInfo: MergeConflictAdditionalInformation = {
+				resolveResult: new ResolveResultAdditionalInfo(result),
+				describeResult: await this.p4.describe(pending.newCl),
+				client: pending.change.targetWorkspace
+			}
+			const verifyResult = await this.sourceNode._evaluateConflictForStomp(pending.newCl, this.targetBranch, additionalInfo)
+			const operationResult = GetStompVerificationResult(verifyResult)
+			if (!operationResult.success) {
+				const failure: Failure = {
+					kind: 'Integration error',
+					description: operationResult.message
+				}
+				return new EdgeIntegrationDetails('error', `${failure.kind}: ${failure.description}`)
+			}
+			result = await this.p4.resolve(
+				pending.change.targetWorkspace,
+				pending.newCl,
+				'clobber',
+				false, // detail === 'quick',
+				edgeServerAddress
 			)
+		}
 
 		if (pending.action.flags.has('manual') || pending.change.forceCreateAShelf) {
 			// the user requested manual merge
@@ -763,7 +788,12 @@ class EdgeBotImpl extends PerforceStatefulBot {
 		let failure: Failure | null = null
 		if (result.hasConflict()) {
 			failure = {
-				kind: 'Merge conflict', description: result.getConflictsText()
+				kind: 'Merge conflict', description: result.getConflictsText(), 
+				additionalInfo: {
+					resolveResult: new ResolveResultAdditionalInfo(result),
+					describeResult: await this.p4.describe(pending.newCl),
+					client: pending.change.targetWorkspace
+				}
 			}
 
 			const conflicts = this.analyzeConflict(result.getConflicts())
diff --git a/Engine/Extras/RoboMerge/v3/src/robo/nodebot.ts b/Engine/Extras/RoboMerge/v3/src/robo/nodebot.ts
--- a/Engine/Extras/RoboMerge/v3/src/robo/nodebot.ts
+++ b/Engine/Extras/RoboMerge/v3/src/robo/nodebot.ts
@@ -9,8 +9,9 @@ import { Change, ConflictedResolveNFile, ExclusiveFileDetails, PerforceContext,
 import { IPCControls, NodeBotInterface, QueuedChange, ReconsiderArgs } from './bot-interfaces';
 import { ApprovalOptions, BotConfig, EdgeOptions, NodeOptions } from './branchdefs'
 import { BlockagePauseInfo, BranchStatus } from './status-types';
-import { AlreadyIntegrated, Blockage, Branch, BranchArg, BranchGraphInterface, ChangeInfo, EndIntegratingToGateEvent, Failure } from './branch-interfaces';
-import { MergeAction, OperationResult, PendingChange, resolveBranchArg, StompedRevision, StompVerification, StompVerificationFile, UnlockVerification }  from './branch-interfaces';
+import { AlreadyIntegrated, Blockage, Branch, BranchArg, BranchGraphInterface, ChangeInfo, EditableBranch, EndIntegratingToGateEvent, Failure } from './branch-interfaces';
+import { GetStompVerificationResult, MergeConflictAdditionalInformation, MergeAction, OperationResult, PendingChange }  from './branch-interfaces';
+import { resolveBranchArg, ResolveResultAdditionalInfo, StompedRevision, StompVerification, StompVerificationFile, UnlockVerification }  from './branch-interfaces';
 import { Conflicts } from './conflicts';
 import { EdgeBot, EdgeIntegrationDetails, EdgeMergeResults } from './edgebot';
 import { BotEventTriggers } from './events';
@@ -628,6 +629,13 @@ export class NodeBot extends PerforceStatefulBot implements NodeBotInterface {
 		return { success: true, message: `Requested CL ${changeCl} successfully sent to bot for reconsideration and shelf creation` }
 	}
 
+	private stompFailureResult(message: string): StompVerification {
+		return {
+			success: false,
+			message,
+			swarmURL: this.p4.swarmURL
+		}
+	}
 	/* 
 	 * In order to verify a stomp request, we need to:
 	 *  1. process the CL using the normal Robomerge process
@@ -639,31 +647,25 @@ export class NodeBot extends PerforceStatefulBot implements NodeBotInterface {
 	async verifyStomp(changeCl: number, targetBranchName: string) : Promise<StompVerification> {
 		this._log_action(`Verifying stomp request of CL ${changeCl} to ${targetBranchName}`, "verbose")
 
-		const failureResult: StompVerification = {
-			success: false,
-			message: "TBD",
-			swarmURL: this.p4.swarmURL
-		}
-
 		if (!changeCl || !targetBranchName) {
-			return { ...failureResult, message: "Missing parameters for stomp node operation" }
+			return this.stompFailureResult("Missing parameters for stomp node operation")
 		}
 
 		const targetBranch = this._getBranch(targetBranchName)
 
 		if (!targetBranch) {
-			return { ...failureResult, message: `Unable to retrieve branch infomation for "${targetBranchName}". Unable to stomp, please contact Robomerge support."` }
+			return this.stompFailureResult(`Unable to retrieve branch infomation for "${targetBranchName}". Unable to stomp, please contact Robomerge support."`)
 		}
 
 		const edge = this.getImmediateEdge(targetBranch)
 
 		if (!edge) {
-			return { ...failureResult, message: `This bot has no edge for specified branch "${targetBranch.name}" on ${this.fullName}`}
+			return this.stompFailureResult(`This bot has no edge for specified branch "${targetBranch.name}" on ${this.fullName}`)
 		}
 
 		// Ensure we're paused on a non-manual pause blockage
 		if (!edge.isBlocked) {
-			return { ...failureResult, message: "Branch needs to have conflict to stomp changes" }
+			return this.stompFailureResult("Branch needs to have conflict to stomp changes")
 		}
 		
 		// Get the pause info of the edge
@@ -671,106 +673,33 @@ export class NodeBot extends PerforceStatefulBot implements NodeBotInterface {
 
 		// We've ensured we're on a blockage now. Make sure the CL is still valid.
 		if (blockageInfo.change !== changeCl) {
-			return { ...failureResult, message: `Requested change CL "${changeCl}" does not match blockage change CL "${blockageInfo.change}"` }
+			return this.stompFailureResult(`Requested change CL "${changeCl}" does not match blockage change CL "${blockageInfo.change}"`)
 		}
 
 		let conflict = this.conflicts.getConflictByChangeCl(changeCl, targetBranch)
 
 		if (!conflict) {
-			return { ...failureResult, message: `Unable to retrieve conflict for CL ${changeCl} merging to ${targetBranch.name}. Unable to stomp, please contact Robomerge support.` }
+			return this.stompFailureResult(`Unable to retrieve conflict for CL ${changeCl} merging to ${targetBranch.name}. Unable to stomp, please contact Robomerge support.`)
 		}
 		if (conflict.kind === 'Exclusive check-out') {
-			return { ...failureResult, message: `Unable to stomp changes over exclusive checkout to files in ${targetBranch.name}. ` + 
-				`Unable to stomp, please revert the locked files in ${targetBranch.name} or contact Robomerge support.` }
-		}
-
-		// Get Change object from P4 for blockage cl
-		const blockageChangeRetrieval = await this._getChange(changeCl)
-
-		if (typeof blockageChangeRetrieval === 'string') {
-			return { 
-				...failureResult, 
-				message: `Unable to retrieve change infomation for given conflict. ${blockageChangeRetrieval}`
-			}
+			return this.stompFailureResult(`Unable to stomp changes over exclusive checkout to files in ${targetBranch.name}. ` + 
+				`Unable to stomp, please revert the locked files in ${targetBranch.name} or contact Robomerge support.`)
 		}
 
-		const blockageChange = blockageChangeRetrieval as Change
-
-		// Set a variety of flags on the change before processing
-		blockageChange.isUserRequest = true
-		blockageChange.forceCreateAShelf = true // We don't want to submit this
-		blockageChange.sendNoShelfNotification = true // Don't send any shelf creation notifications for this
+		return this._evaluateConflictForStomp(changeCl, targetBranch, blockageInfo.additionalInfo)
+	}
 
-		// Process and attempt the integration
-		const edgeMap = new Map<string, EdgeBot>([[edge.targetBranch.upperName, edge]])
-		let integratedCl: Change | null = null // Set to possible undefines as we use this in the finally clause
+	async _evaluateConflictForStomp(changeCl: number, targetBranch: Readonly<EditableBranch>, mergeConflictInfo: MergeConflictAdditionalInformation) : Promise<StompVerification>
+	{
 		try {
-			const workspace = await edge.getWorkspace()
-			const processResult = await this._createChangeInfo(blockageChange, edgeMap, this.p4.username, workspace, targetBranch)
-			if (!processResult.info) {
-				return {
-					...failureResult,  
-					message: `Could not parse blockage change\nErrors: ${processResult.errors}`
-				}
-			}
-			const mergeResult = await this._mergeClViaEdges(edgeMap, processResult.info, true)
-
-			// If we don't have a change returned, we ran into an error processing the CL
-			const changes = mergeResult.allChanges()
-			if (changes.length !== 1) { // Encountered when we get errors back
-				let message = `Got incorrect number of entries back from integrating CL ${changeCl} (length = ${changes.length}, should be 1)`
-				if (mergeResult.errors.length > 0) {
-					message = `${message}\nErrors: ${mergeResult.errors}`
-				}
-				return { 
-					...failureResult,  
-					message
-				}
-			}
-
-			const changelist = changes[0]!
-			if (changelist <= 0) {
-				return { 
-					...failureResult, 
-					message: `Unable to get integration shelf for CL ${changeCl}`
-				}
-			}
-
-			let p4ChangeResult
-			try {
-				p4ChangeResult = await this.p4.getChange(changelist)
-			}
-			catch (err) {
-				return { 
-					...failureResult, 
-					message: `Unable to get integration shelf for CL ${changeCl}: ${err}`
-				}
-			}
-
-			integratedCl = p4ChangeResult as Change
-
-			// Unshelve 
-			if (!await this.p4.unshelve(workspace, integratedCl.change)) {
-				return { 
-					...failureResult,  
-					message: `Unable to unshelve CL ${integratedCl.change}`
-				}
-			}
-
 			// Get what is left to resolve out of the changelist
-			const resolveResult = await this.p4.resolve(integratedCl.client, integratedCl.change, 'normal')
+			const resolveResult = Object.assign(new ResolveResultAdditionalInfo(), mergeConflictInfo.resolveResult)
 
 			if (!resolveResult.hasConflict()) {
-				return { 
-					...failureResult,  
-					message: `No conflict found resolving CL ${changeCl}`
-				}
+				return this.stompFailureResult(`No conflict found resolving CL ${changeCl}`)
 			}
 			if (!resolveResult.successfullyParsedConflicts) {
-				return { 
-					...failureResult,  
-					message: `Encountered error enumerating conflicting files after integration attempt: ${resolveResult.getParsingError()}`
-				}
+				return this.stompFailureResult(`Encountered error enumerating conflicting files after integration attempt: ${resolveResult.parsingError}`)
 			}
 
 
@@ -790,13 +719,11 @@ export class NodeBot extends PerforceStatefulBot implements NodeBotInterface {
 			let svFiles : StompVerificationFile[] = []
 
 			// Get information about the pending changelist as a whole
-			const describeResult = await this.p4.describe(integratedCl.change)
-
 			const remainingFiles = resolveResult.getConflicts()
 			// We need the depot file for each conflict file to do comparisons
 			for (const remainingFile of remainingFiles) {
 				// Find target file 
-				const targetFileInDepot = await this.p4.where(integratedCl.client, remainingFile.clientFile)
+				const targetFileInDepot = await this.p4.where(mergeConflictInfo.client, remainingFile.clientFile)
 				if (!targetFileInDepot[0].depotFile) { 
 					throw new Error(`Error retrieving depot path for the merge target of ${remainingFile.clientFile}`)
 				}
@@ -817,7 +744,7 @@ export class NodeBot extends PerforceStatefulBot implements NodeBotInterface {
 			
 			// Process all the files in our pending changelist
 			let stompedRevisionsCalculationIssues = false
-			for (const file of describeResult.entries) {
+			for (const file of mergeConflictInfo.describeResult.entries) {
 				// Check if this file is still remaining in the changelist with a conflict
 				const remainingFile = getRemainingFile(file.depotFile, remainingFiles)
 				const resolved = !remainingFile
@@ -856,8 +783,8 @@ export class NodeBot extends PerforceStatefulBot implements NodeBotInterface {
 							debugScope.setExtra('startFromRev', remainingFile!.startFromRev.toString())
 							debugScope.setExtra('endFromRev', remainingFile!.endFromRev.toString())
 							debugScope.setExtra('remainingFileJson', JSON.stringify(remainingFile!))
-							debugScope.setExtra('resolveResult', resolveResult!.getConflictsText())
-							debugScope.setExtra('resolveResultParsingError?', resolveResult!.getParsingError())
+							debugScope.setExtra('resolveResult', resolveResult!.conflictsText)
+							debugScope.setExtra('resolveResultParsingError?', resolveResult!.parsingError)
 						}
 
 						// Attempt to find cause of NaN errors
@@ -905,11 +832,11 @@ export class NodeBot extends PerforceStatefulBot implements NodeBotInterface {
 			}
 
 			if (!remainingAllBinary) {
-				returnResult.message= "Invalid Request: Not all remaining files are binary"
+				returnResult.message = "Invalid Request: Not all remaining files are binary"
 				returnResult.validRequest = false
 			}
 			else if (stompedRevisionsCalculationIssues) {
-				returnResult.message= "Valid Request, but encountered issues calculating stomped revisions on some files"
+				returnResult.message = "Valid Request, but encountered issues calculating stomped revisions on some files"
 				returnResult.validRequest = true
 				
 				if (this.tickJournal) {
@@ -918,27 +845,16 @@ export class NodeBot extends PerforceStatefulBot implements NodeBotInterface {
 				}
 			}
 			else {
-				returnResult.message= "Valid Request"
+				returnResult.message = "Valid Request"
 				returnResult.validRequest = true
 			}
 
 			// Finally return the results
 			return returnResult
-
 			
 		} catch (error) {
-			return { 
-				success: false, 
-				message: `Encountered error processing ${changeCl}: ${error}`
-			}
+			return this.stompFailureResult(`Encountered error processing ${changeCl}: ${error}`)
 		} 
-		finally {
-			// Delete the CL, including any remaining shelved files
-			this.nodeBotLogger.info(`${this.fullName} - Cleaning up after verifyStompRequest()`)
-			if (integratedCl) {
-				await edge.revertPendingCLWithShelf(integratedCl.client, integratedCl.change, 'Stomp Verify')
-			}
-		}
 	}
 
 	// TODO: Support branch/delete merges
@@ -1006,18 +922,9 @@ export class NodeBot extends PerforceStatefulBot implements NodeBotInterface {
 
 		// Reverify, but don't revert on a successful verification!
 		const verifyResult = await this.verifyStomp(changeCl, targetBranchName)
-
-		if (!verifyResult.success) {
-			return { success: false, message: `stompChanges Request Failed Due to Bad Verification: ${verifyResult.message}` }
-		}
-
-		if (!verifyResult.validRequest) {
-			return { success: false, message: `stompChanges Request Cannot Be Completed: ${verifyResult.message}` }
-		}
-
-		// One last check to ensure this is an all binary stomp
-		if (!verifyResult.remainingAllBinary) {
-			return { success: false, message: `stompChanges Failed -- Verification had remaining text files (${verifyResult.message})` }
+		const operationResult = GetStompVerificationResult(verifyResult)
+		if (!operationResult.success) {
+			return operationResult
 		}
 
 		// Before we submit, add an #fyi for all the stomped revision authors in the changelist
@@ -1041,8 +948,7 @@ export class NodeBot extends PerforceStatefulBot implements NodeBotInterface {
 			++this.tickJournal.stompQueued
 		}
 
-		return { success: true, message: 'Queued stomp request' }
-
+		return operationResult
 	}
 
 	async verifyUnlock(changeCl: number, targetBranchName: string) : Promise<UnlockVerification> {
