Help Page is now configurable via environment variable
Fixed URLs for wiki links post migration

diff --git a/Engine/Extras/RoboMerge/v3/public/js/boilerplate.js b/Engine/Extras/RoboMerge/v3/public/js/boilerplate.js
--- a/Engine/Extras/RoboMerge/v3/public/js/boilerplate.js
+++ b/Engine/Extras/RoboMerge/v3/public/js/boilerplate.js
@@ -9,6 +9,7 @@
 
 // Global user information map
 let robomergeUser
+let helpPageURL = '/help'
 
 window.showNodeLastChanges = false
 window.allowNodeReconsiders = false
@@ -437,7 +438,7 @@ function generateRobomergeHeader(createSignedInUserDiv = true) {
 
 	const bigHelpButton = $('<button>')
 		.addClass('btn btn-outline-dark big-help')
-		.click(function() { window.location.href='/help' })
+		.click(function() { window.open(helpPageURL, "_blank") })
 		.html('<strong>Help Page</strong>')
 	topright.append(bigHelpButton)
 
@@ -562,13 +563,13 @@ function generateRobomergeFooter() {
 
 	let helpButton = $('<button id="helpButton">')
 	helpButton.addClass("btn btn-sm btn-outline-dark")
-	helpButton.click(function() { window.location.href='/help' })
+	helpButton.click(function() { window.open(helpPageURL, "_blank") })
 	helpButton.html('<strong>Help Page</strong>')
 	rightButtonDiv.append(helpButton)
 
 	let contactButton = $('<button id="contactButton">')
 	contactButton.addClass("btn btn-sm btn-outline-dark")
-	contactButton.click(function() { window.location.href='/contact' })
+	contactButton.click(function() { window.open('/contact', "_blank") })
 	contactButton.html('Contact Us')
 	rightButtonDiv.append(contactButton)
 	
@@ -1940,6 +1941,10 @@ function preprocessData(data) {
 
 const onLoginPage = window.location.pathname === "/login"
 
+$(async () => {
+	helpPageURL = await $.get('/helpPageURL')
+})
+
 // Define some standard header and footer for our pages
 generateRobomergeHeader(!onLoginPage)
 generateRobomergeFooter()
diff --git a/Engine/Extras/RoboMerge/v3/src/robo/notifications.ts b/Engine/Extras/RoboMerge/v3/src/robo/notifications.ts
--- a/Engine/Extras/RoboMerge/v3/src/robo/notifications.ts
+++ b/Engine/Extras/RoboMerge/v3/src/robo/notifications.ts
@@ -805,9 +805,11 @@ export class BotNotifications implements BotEventHandler {
 		// Start collecting our attachments
 		let attachCollection : SlackAttachment[] = []
 		
+		const helpURL = args.helpPageURL != '/help' ? args.helpPageURL : `${this.externalRobomergeUrl}/help`
+
 		// General information
 		attachCollection.push({
-			text: `To learn more about robomerge and how to resolve blockages, please review the <${this.externalRobomergeUrl}/help|robomerge help page>.`,
+			text: `To learn more about robomerge and how to resolve blockages, please review the <${helpURL}|robomerge help page>.`,
 			mrkdwn_in: ["text"]
 		})
 
diff --git a/Engine/Extras/RoboMerge/v3/src/robo/robo.ts b/Engine/Extras/RoboMerge/v3/src/robo/robo.ts
--- a/Engine/Extras/RoboMerge/v3/src/robo/robo.ts
+++ b/Engine/Extras/RoboMerge/v3/src/robo/robo.ts
@@ -155,6 +155,12 @@ const COMMAND_LINE_ARGS: {[param: string]: Arg<any>} = {
 		dflt: 100
 	},
 
+	helpPageURL: {
+		match: /^-helpPageURL=(.+)$/,
+		env: 'HELP_PAGE_URL',
+		dflt: '/help'
+	},
+
 	// Sentry environment designation -- use 'PROD' to enable Sentry bug tracking
 	epicEnv: {
 		match: /^(-epicEnv)$/,
diff --git a/Engine/Extras/RoboMerge/v3/src/robo/roboserver.ts b/Engine/Extras/RoboMerge/v3/src/robo/roboserver.ts
--- a/Engine/Extras/RoboMerge/v3/src/robo/roboserver.ts
+++ b/Engine/Extras/RoboMerge/v3/src/robo/roboserver.ts
@@ -47,6 +47,11 @@ let ENVIRONMENT: {[param: string]: any}
 			env: 'ROBO_USE_AUTH_IN_DEV',
 			dflt: false
 		},
+		helpPageURL: {
+			match: /^(-helpPageURL)$/,
+			env: 'HELP_PAGE_URL',
+			dflt: '/help'
+		},
 		oktaSignIn: {
 			match: /^(-oktaSignIn)$/,
 			parse: str => str === "false" ? false : true,
@@ -436,6 +441,11 @@ class RoboWebApp implements AppInterface {
 		return signIn;
 	}
 
+	@Handler('GET', '/helpPageURL')
+	async helpPageURL() {
+		return ENVIRONMENT.helpPageURL
+	}
+
 	@SecureHandler('PUT', '/api/control/verbose/*')
 	async setVerbose(onOff: string): Promise<OperationReturnType> {
 		if (!this.authData) {
