Detect duplicate entries in flowsTo, forceflowsTo, defaultFlow, and graphFilteredEdges
Add set union/difference/intersection helpers since getting the functional tests running on 5.5+ to use Set.difference/intersection/union is apparently an impossible task

diff --git a/Engine/Extras/RoboMerge/v3/src/common/helper.ts b/Engine/Extras/RoboMerge/v3/src/common/helper.ts
--- a/Engine/Extras/RoboMerge/v3/src/common/helper.ts
+++ b/Engine/Extras/RoboMerge/v3/src/common/helper.ts
@@ -91,3 +91,31 @@ export class Random {
 		return container[Random.randint(0, container.length - 1)]
 	}
 }
+
+export function setDifference<T>(setA: Set<T>, setB: Set<T>): Set<T> {
+  const result = new Set<T>();
+  for (const item of setA) {
+    if (!setB.has(item)) {
+      result.add(item);
+    }
+  }
+  return result;
+}
+
+export function setIntersection<T>(setA: Set<T>, setB: Set<T>): Set<T> {
+  const result = new Set<T>();
+  for (const item of setA) {
+    if (setB.has(item)) {
+      result.add(item);
+    }
+  }
+  return result;
+}
+
+export function setUnion<T>(setA: Set<T>, setB: Set<T>): Set<T> {
+  const result = new Set<T>(setA); // Start with all elements from setA
+  for (const item of setB) {
+    result.add(item); // Add elements from setB (duplicates automatically skipped)
+  }
+  return result;
+}
\ No newline at end of file
diff --git a/Engine/Extras/RoboMerge/v3/src/robo/branchdefs.ts b/Engine/Extras/RoboMerge/v3/src/robo/branchdefs.ts
--- a/Engine/Extras/RoboMerge/v3/src/robo/branchdefs.ts
+++ b/Engine/Extras/RoboMerge/v3/src/robo/branchdefs.ts
@@ -1,7 +1,8 @@
 // Copyright Epic Games, Inc. All Rights Reserved.
 
+import { setDifference } from '../common/helper'
 import { PerforceContext } from '../common/perforce'
-import { ContextualLogger } from '../common/logger';
+import { ContextualLogger } from '../common/logger'
 
 const jsonlint: any = require('jsonlint-mod')
 
@@ -505,39 +506,53 @@ export class BranchDefs {
 			}))
 		}
 
-		// Check flow
-		for (const def of branchGraph.branches) {
-			const flowsTo = new Set()
-			if (def.flowsTo) {
-				if (!Array.isArray(def.flowsTo)) {
-					validationErrors.push(`'${def.name}'.flowsTo is not an array`)
+		const validateBranchArray = function(arr: string[]|undefined, defName: string|undefined, propName: string): Set<string> {
+			const uniqueValues = new Set<string>()
+			if (arr) {
+				if (!Array.isArray(arr)) {
+					validationErrors.push(`'${defName}'.${propName} is not an array`)
 				}
-				else for (const to of def.flowsTo) {
+				else for (const to of arr) {
 					const branchName = names.get(to.toUpperCase())
 					if (branchName) {
-						flowsTo.add(branchName)
+						if (uniqueValues.has(branchName)) {
+							validationErrors.push(`'${defName}'.${propName} has duplicate entries for '${to}'`)
+						} else {
+							uniqueValues.add(branchName)
+						}
 					}
 					else {
-						validationErrors.push(`'${def.name}' flows to unknown branch/alias '${to}'`)
+						validationErrors.push(`'${defName}'.${propName} has unknown branch/alias '${to}'`)
 					}
-				}
+				}				
 			}
+			return uniqueValues
+		}
 
-			if (def.forceFlowTo) {
-				if (!Array.isArray(def.forceFlowTo)) {
-					validationErrors.push(`'${def.name}'.forceFlowTo is not an array`)
-				}
-				else for (const to of def.forceFlowTo) {
-					const branchName = names.get(to.toUpperCase())
-					if (!branchName) {
-						validationErrors.push(`'${def.name}' force flows to unknown branch/alias '${to}'`)
-					}
-					else if (!flowsTo.has(branchName)) {
-						validationErrors.push(`'${def.name}' force flows but does not flow to '${to}'`)
-					}
-				}
+		// Check flow
+		for (const def of branchGraph.branches) {
+			const flowsTo = validateBranchArray(def.flowsTo, def.name, 'flowsTo')
+
+			const forceFlowsTo = validateBranchArray(def.forceFlowTo, def.name, 'forceFlowTo')
+			const missingForceFlows = setDifference(forceFlowsTo, flowsTo)
+			for (const missingFlow of missingForceFlows) {
+				validationErrors.push(`'${def.name}' force flows but does not flow to '${missingFlow}'`)
 			}
 
+			const blockAsset = validateBranchArray(def.blockAssetFlow, def.name, 'blockAssetFlow')
+			const missingBlockAsset = setDifference(blockAsset, flowsTo)
+			for (const missingFlow of missingBlockAsset) {
+				validationErrors.push(`'${def.name}' blocks assets but does not flow to '${missingFlow}'`)
+			}
+
+			const filteredEdges = validateBranchArray(def.graphFilteredEdges, def.name, 'graphFilteredEdges')
+			const missingFilteredEdges = setDifference(filteredEdges, flowsTo)
+			for (const missingFlow of missingFilteredEdges) {
+				validationErrors.push(`'${def.name}' filters edges from the graph but does not flow to '${missingFlow}'`)
+			}
+
+			validateBranchArray(def.defaultFlow, def.name, 'defaultFlow')
+
 			if (def.changelistParser) {
 				try {
 					await import(`./changelistparsers/${def.changelistParser}`)
