Spelling and commented code clean up

diff --git a/Engine/Extras/RoboMerge/v3/functional_tests/framework/src/tests/test-edge-reconsider.ts b/Engine/Extras/RoboMerge/v3/functional_tests/framework/src/tests/test-edge-reconsider.ts
--- a/Engine/Extras/RoboMerge/v3/functional_tests/framework/src/tests/test-edge-reconsider.ts
+++ b/Engine/Extras/RoboMerge/v3/functional_tests/framework/src/tests/test-edge-reconsider.ts
@@ -36,9 +36,6 @@ export class TestEdgeReconsider extends FunctionalTest {
 			this.reconsider('Main', this.revisionCl, 'Dev-Perkin'),
 			this.reconsider('Main', this.conflictingRevisionCl, 'Dev-Posie')
 		])
-
-			// this.reconsider('Main', this.revisionCl, 'Dev-Perkin')
-			// .then(() => this.reconsider('Main', this.conflictingRevisionCl, 'Dev-Posie'))
 	}
 
 	verify() {
diff --git a/Engine/Extras/RoboMerge/v3/public/js/flow.graph.js b/Engine/Extras/RoboMerge/v3/public/js/flow.graph.js
--- a/Engine/Extras/RoboMerge/v3/public/js/flow.graph.js
+++ b/Engine/Extras/RoboMerge/v3/public/js/flow.graph.js
@@ -14,7 +14,7 @@ class Info {
     }
     update(branch, getId) {
         const edgeId = (alias) => branch.def.name + '->' + getId(alias);
-        // shuld really do the same for default and block assets
+        // should really do the same for default and block assets
         this.filteredEdges = new Set([...this.filteredEdges, ...branch.def.graphFilteredTargets.map(edgeId)]);
         this.forcedEdges = new Set([...this.forcedEdges, ...branch.def.forceFlowTo.map(edgeId)]);
         this.defaultDests = new Set([...this.defaultDests, ...branch.def.defaultFlow.map(edgeId)]);
