Block Asset Reconsider allows a node to specify targets that are forceFlow and want assets when doing so, but want to prevent assets from being reconsidered past gates/pause.

diff --git a/Engine/Extras/RoboMerge/v3/README.md b/Engine/Extras/RoboMerge/v3/README.md
--- a/Engine/Extras/RoboMerge/v3/README.md
+++ b/Engine/Extras/RoboMerge/v3/README.md
@@ -102,6 +102,7 @@ for both node and edge.
 |`forcePause`               |ne|            | |If flag set, pause - applies each restart| |
 |`additionalSlackChannelForBlockages` |e |  | |Single extra Slack channel               | |
 |`blockAssetFlow`           |n |`[]`        | |Reject any integrations containing assets to these targets   | |
+|`blockAssetReconsider`     |n |`[]`        | |Reject reconsider requests containing assets to these targets   | |
 |`graphFilteredEdges`       |n |`[]`        | |By default these edges will not be displayed on the graph  | |
 |`integrationMethod`        |e |normal      | |For skip if unchanged                    | |
 |`disallowSkip`             |e |            | |Remove skip option from UI               | |
diff --git a/Engine/Extras/RoboMerge/v3/functional_tests/framework/src/framework.ts b/Engine/Extras/RoboMerge/v3/functional_tests/framework/src/framework.ts
--- a/Engine/Extras/RoboMerge/v3/functional_tests/framework/src/framework.ts
+++ b/Engine/Extras/RoboMerge/v3/functional_tests/framework/src/framework.ts
@@ -150,6 +150,7 @@ type RobomergeBranchOptions = {
 	forceAll: boolean
 	visibility: string[] | string
 	blockAssetFlow: string[]
+	blockAssetReconsider: string[]
 	disallowDeadend: boolean
 
 	streamDepot: string
diff --git a/Engine/Extras/RoboMerge/v3/functional_tests/framework/src/tests.ts b/Engine/Extras/RoboMerge/v3/functional_tests/framework/src/tests.ts
--- a/Engine/Extras/RoboMerge/v3/functional_tests/framework/src/tests.ts
+++ b/Engine/Extras/RoboMerge/v3/functional_tests/framework/src/tests.ts
@@ -44,11 +44,12 @@ import { StompWithBranch } from './tests/stomp-with-branch'
 import { StompWithDelete } from './tests/stomp-with-delete'
 import { SyntaxErrorOnUnknownBranch } from './tests/syntax-error-on-unknown-branch'
 import { TestChain } from './tests/test-chain'
+import { TestEdgeBlockAssetsReconsider } from './tests/test-edge-block-assets-reconsider'
 import { TestEdgeGate } from './tests/test-edge-gate'
+import { TestEdgeReconsider } from './tests/test-edge-reconsider'
 import { TestFlags } from './tests/test-flags'
 import { TestGate } from './tests/test-gate'
 import { TestReconsider } from './tests/test-reconsider'
-import { TestEdgeReconsider } from './tests/test-edge-reconsider'
 import { TestTerminal } from './tests/test-terminal'
 import { UnreachableSkip } from './tests/unreachable-skip'
 
@@ -228,7 +229,8 @@ async function go() {
 		new ExcludeDescriptions(p4), // 50
 
 		new ExcludeDescriptionsPerEdge(p4),
-		new SkipIfUnchanged(p4)
+		new SkipIfUnchanged(p4),
+		new TestEdgeBlockAssetsReconsider(p4)
 	]
 
 	const TARGET_TEST_DEFS: [string[], string | null][] = [
diff --git a/Engine/Extras/RoboMerge/v3/functional_tests/framework/src/tests/test-edge-block-assets-reconsider.ts b/Engine/Extras/RoboMerge/v3/functional_tests/framework/src/tests/test-edge-block-assets-reconsider.ts
new file mode 100644
--- /dev/null
+++ b/Engine/Extras/RoboMerge/v3/functional_tests/framework/src/tests/test-edge-block-assets-reconsider.ts
@@ -0,0 +1,58 @@
+// Copyright Epic Games, Inc. All Rights Reserved.
+import { FunctionalTest, P4Util, Stream } from '../framework'
+
+const streams: Stream[] = [
+	{name: 'Main', streamType: 'mainline'},
+	{name: 'Dev-Pootle', streamType: 'development', parent: 'Main'},
+	{name: 'Dev-Perkin', streamType: 'development', parent: 'Main'},
+	{name: 'Dev-Posie', streamType: 'development', parent: 'Main'}
+]
+
+export class TestEdgeBlockAssetsReconsider extends FunctionalTest {
+
+	async setup() {
+		await this.p4.depot('stream', this.depotSpec())
+		await this.createStreamsAndWorkspaces(streams)
+
+		await P4Util.addFileAndSubmit(this.getClient('Main'), 'test.uasset', 'Initial content')
+
+		const desc = 'Initial branch of files from Main'
+		await Promise.all([
+			this.p4.populate(this.getStreamPath('Dev-Perkin'), desc),
+			this.p4.populate(this.getStreamPath('Dev-Pootle'), desc),
+			this.p4.populate(this.getStreamPath('Dev-Posie'), desc)
+		])
+	}
+
+	async run() {
+		// making edit in set-up, i.e. before added to RoboMerge
+		const revisionCl = await P4Util.editFileAndSubmit(this.getClient('Main'), 'test.uasset', 'Initial content\n\nFirst addition')
+
+        return Promise.all([
+			this.reconsider('Main', revisionCl, 'Dev-Perkin'),
+			this.reconsider('Main', revisionCl, 'Dev-Posie', `#robomerge ${this.fullBranchName('Dev-Posie')} | #robomerge #disregardAssetBlock`)
+		])
+	}
+
+	verify() {
+		return Promise.all([
+			this.checkHeadRevision('Main', 'test.uasset', 2),
+			this.checkHeadRevision('Dev-Pootle', 'test.uasset', 2),
+			this.checkHeadRevision('Dev-Perkin', 'test.uasset', 1),
+			this.checkHeadRevision('Dev-Posie', 'test.uasset', 2)
+		])
+	}
+
+	getBranches() {
+		const mainDef = this.makeBranchDef('Main', ['Dev-Perkin', 'Dev-Pootle', 'Dev-Posie'])
+        mainDef.forceFlowTo = [this.fullBranchName('Dev-Pootle')]
+		mainDef.blockAssetReconsider = [this.fullBranchName('Dev-Perkin'), this.fullBranchName('Dev-Pootle'), this.fullBranchName('Dev-Posie')]
+		return [
+			mainDef,
+			this.makeBranchDef('Dev-Perkin', []),
+			this.makeBranchDef('Dev-Pootle', []),
+			this.makeBranchDef('Dev-Posie', [])
+        ]
+	}
+
+}
diff --git a/Engine/Extras/RoboMerge/v3/src/robo/branch-interfaces.ts b/Engine/Extras/RoboMerge/v3/src/robo/branch-interfaces.ts
--- a/Engine/Extras/RoboMerge/v3/src/robo/branch-interfaces.ts
+++ b/Engine/Extras/RoboMerge/v3/src/robo/branch-interfaces.ts
@@ -106,6 +106,7 @@ export interface BranchGraphInterface {
 	getBranchNames(): string[]
 }
 
+export type BlockAssetType = 'all' | 'reconsider'
 export type EditableBranch = BranchBase & {
 	bot?: NodeBotInterface
 	parent: BranchGraphInterface
@@ -119,7 +120,7 @@ export type EditableBranch = BranchBase & {
 	forcedDownstream?: string[]
 	enabled: boolean
 	visibility: string[] | string
-	blockAssetTargets: Set<string>
+	blockAssetTargets: Map<string, BlockAssetType>
 	graphFilteredTargets: Set<string>
 	allowDeadend: boolean
 
diff --git a/Engine/Extras/RoboMerge/v3/src/robo/branchdefs.ts b/Engine/Extras/RoboMerge/v3/src/robo/branchdefs.ts
--- a/Engine/Extras/RoboMerge/v3/src/robo/branchdefs.ts
+++ b/Engine/Extras/RoboMerge/v3/src/robo/branchdefs.ts
@@ -1,6 +1,6 @@
 // Copyright Epic Games, Inc. All Rights Reserved.
 
-import { setDifference } from '../common/helper'
+import { setDifference, setIntersection } from '../common/helper'
 import { PerforceContext } from '../common/perforce'
 import { ContextualLogger } from '../common/logger'
 
@@ -116,6 +116,7 @@ const nodeOptionFieldsPrototype = {
 	forceAll: false,
 	visibility: '' as string[] | string,
 	blockAssetFlow: [''],
+	blockAssetReconsider: [''],
 	graphFilteredEdges: [''],
 	disallowDeadend: false,
 
@@ -539,12 +540,22 @@ export class BranchDefs {
 				validationErrors.push(`'${def.name}' force flows but does not flow to '${missingFlow}'`)
 			}
 
-			const blockAsset = validateBranchArray(def.blockAssetFlow, def.name, 'blockAssetFlow')
-			const missingBlockAsset = setDifference(blockAsset, flowsTo)
-			for (const missingFlow of missingBlockAsset) {
+			const blockAssetFlow = validateBranchArray(def.blockAssetFlow, def.name, 'blockAssetFlow')
+			const missingBlockAssetFlows = setDifference(blockAssetFlow, flowsTo)
+			for (const missingFlow of missingBlockAssetFlows) {
 				validationErrors.push(`'${def.name}' blocks assets but does not flow to '${missingFlow}'`)
 			}
 
+			const blockAssetReconsider = validateBranchArray(def.blockAssetReconsider, def.name, 'blockAssetReconsider')
+			const missingBlockAssetReconsiders = setDifference(blockAssetReconsider, flowsTo)
+			for (const missingFlow of missingBlockAssetReconsiders) {
+				validationErrors.push(`'${def.name}' blocks asset reconsiders but does not flow to '${missingFlow}'`)
+			}
+			const overlappingBlockAssets = setIntersection(blockAssetReconsider, blockAssetFlow)
+			for (const overlappingBlock of overlappingBlockAssets) {
+				validationErrors.push(`'${def.name}' blocks asset reconsiders to '${overlappingBlock}' but already blocks all asset flow`)
+			}
+
 			const filteredEdges = validateBranchArray(def.graphFilteredEdges, def.name, 'graphFilteredEdges')
 			const missingFilteredEdges = setDifference(filteredEdges, flowsTo)
 			for (const missingFlow of missingFilteredEdges) {
diff --git a/Engine/Extras/RoboMerge/v3/src/robo/branchgraph.ts b/Engine/Extras/RoboMerge/v3/src/robo/branchgraph.ts
--- a/Engine/Extras/RoboMerge/v3/src/robo/branchgraph.ts
+++ b/Engine/Extras/RoboMerge/v3/src/robo/branchgraph.ts
@@ -1,8 +1,8 @@
 // Copyright Epic Games, Inc. All Rights Reserved.
 
-import { setDefault } from '../common/helper';
+import { setDefault, setDifference, setIntersection, setUnion } from '../common/helper';
 import { BranchSpec } from '../common/perforce';
-import { Branch, BranchGraphInterface, EditableBranch } from './branch-interfaces';
+import { BlockAssetType, Branch, BranchGraphInterface, EditableBranch } from './branch-interfaces';
 import { BotConfig, BranchDefs, BranchGraphDefinition, EdgeOptions, EdgeProperties } from './branchdefs';
 import { NodeOptions, calculateStream } from './branchdefs';
 
@@ -238,7 +238,7 @@ if (botname === '__TEST__') {
 			forceFlowTo: [], // will compute
 			defaultFlow: options.defaultFlow || [],
 			macros: options.macros,
-			blockAssetTargets: new Set<string>(),
+			blockAssetTargets: new Map<string, BlockAssetType>(),
 			graphFilteredTargets: new Set<string>(),
 			enabled: !options.disabled,
 			allowDeadend: !options.disallowDeadend,
@@ -277,58 +277,81 @@ if (botname === '__TEST__') {
 			branch.stream = streamResult.stream
 		}
 
-		// make sure flowsTo and defaultFlow are all upper case
-		for (let i=0; i<branch.flowsTo.length; ++i) {
-			branch.flowsTo[i] = branch.flowsTo[i].toUpperCase()
-		}
-
-		if (options.blockAssetFlow) {
-			for (const flow of options.blockAssetFlow) {
-				const flowUpper = flow.toUpperCase()
-				if (branch.flowsTo.indexOf(flowUpper) === -1) {
-					throw new Error(`Block asset flow '${flow}' not in flowsTo list`)
+		const processBranchArray = function(arr: string[], propName: string) {
+			const uniqueValues = new Set<string>()
+			const duplicateValues: string[] = []
+			for (const i in arr) {
+				arr[i] = arr[i].toUpperCase()
+				if (uniqueValues.has(arr[i])) {
+					duplicateValues.push(arr[i])
+				} else {
+					uniqueValues.add(arr[i])
 				}
-				branch.blockAssetTargets.add(flowUpper)
 			}
+			if (duplicateValues.length > 0) {
+				throw new Error(`Found duplicate entries [${duplicateValues.join(', ')}] in ${propName}`)
+			}
+			return uniqueValues
 		}
 
-		if (options.graphFilteredEdges) {
-			for (const filtered of options.graphFilteredEdges) {
-				const flowUpper = filtered.toUpperCase()
-				if (branch.flowsTo.indexOf(flowUpper) === -1) {
-					throw new Error(`Graph filtered edge '${filtered}' not in flowsTo list`)
-				}
-				branch.graphFilteredTargets.add(flowUpper)
+		const FlowsToBranches = processBranchArray(branch.flowsTo, 'flowsTo')
 
+		const checkMissingFlowsTo = function(branches: Set<string>, msgPrefix: string) {
+			const MissingFlowsTo = setDifference(branches, FlowsToBranches)
+			if (MissingFlowsTo.size == 1) {
+				throw new Error(`${msgPrefix} '${[...MissingFlowsTo][0]}' not in flowsTo list`)
+			} else if (MissingFlowsTo.size > 1) {
+				throw new Error(`${msgPrefix}s [${[...MissingFlowsTo].join(', ')}] not in flowsTo list`)
 			}
 		}
 
-		for (let i=0; i<branch.defaultFlow.length; ++i) {
-			branch.defaultFlow[i] = branch.defaultFlow[i].toUpperCase()
+		let BlockAssetFlowBranches: Set<string>|null = null
+		if (options.blockAssetFlow) {
+			BlockAssetFlowBranches = processBranchArray(options.blockAssetFlow, 'blockAssetFlow')
+			checkMissingFlowsTo(BlockAssetFlowBranches, "Block asset flow")
+			BlockAssetFlowBranches.forEach(v => branch.blockAssetTargets.set(v, 'all'))
 		}
 
-		// apply forced merge targets
-		const forceFlow = (options.forceAll ? options.flowsTo : options.forceFlowTo) || []
-		for (const target of forceFlow) {
-			const upperTarget = target.toUpperCase()
-
-			// make sure anything in force is in flowsTo
-			if (branch.flowsTo.indexOf(upperTarget) < 0) {
-				throw new Error(`Cannot force ${branch.name} to flow to ${upperTarget} because it is not in the flowsTo list`)
+		if (options.blockAssetReconsider) {
+			const BlockAssetReconsiderBranches = processBranchArray(options.blockAssetReconsider, 'blockAssetReconsider')
+			checkMissingFlowsTo(BlockAssetReconsiderBranches, "Block asset reconsider")
+			if (BlockAssetFlowBranches) {
+				const DuplicateBlockAssets = setIntersection(BlockAssetReconsiderBranches, BlockAssetFlowBranches)
+				if (DuplicateBlockAssets.size == 1) {
+					throw new Error(`Block asset reconsider '${[...DuplicateBlockAssets][0]}' already in block asset flow list`)
+				} else if (DuplicateBlockAssets.size > 1) {
+					throw new Error(`Block asset reconsiders [${[...DuplicateBlockAssets].join(', ')}] already in block asset flow list`)
+				}
 			}
-			branch.forceFlowTo.push(upperTarget)
+			BlockAssetReconsiderBranches.forEach(v => branch.blockAssetTargets.set(v, 'reconsider'))
+		}
+
+		if (options.graphFilteredEdges) {
+			branch.graphFilteredTargets = processBranchArray(options.graphFilteredEdges, 'graphFilteredEdges')
+			checkMissingFlowsTo(branch.graphFilteredTargets, "Graph filtered edge")
+		}
 
-			// make sure anything in forceFlowTo also exists in defaultFlow
-			if (branch.defaultFlow.indexOf(upperTarget) < 0) {
-				branch.defaultFlow.push(upperTarget)
+		const DefaultFlowBranches = processBranchArray(branch.defaultFlow, 'defaultFlow')
+
+		// apply forced merge targets
+		{
+			const forceFlow = (options.forceAll ? options.flowsTo : options.forceFlowTo) || []
+			const ForceFlowBranches = processBranchArray(forceFlow, 'forceFlowTo')
+			const MissingFlowsTo = setDifference(ForceFlowBranches, FlowsToBranches)
+			if (MissingFlowsTo.size == 1) {
+				throw new Error(`Cannot force ${branch.name} to flow to '${[...MissingFlowsTo][0]}' because it is not in the flowsTo list`)
+			} else if (MissingFlowsTo.size > 1) {
+				throw new Error(`Cannot force ${branch.name} to flow to [${[...MissingFlowsTo].join(', ')}] because they are not in the flowsTo list`)
 			}
+			branch.forceFlowTo = [...ForceFlowBranches]
+			branch.defaultFlow = [...setUnion(ForceFlowBranches, DefaultFlowBranches)]
 		}
 
 		// add to branches
 		this.branches.push(branch)
 
 		// add to names
-		this.names.set(name.toUpperCase(), branch)
+		this.names.set(nameUpper, branch)
 
 // alias problem: we're storing flows by name provided, which may be an alias
 //	- works, because above just checks names (e.g. flowsTo and forceFlow)
diff --git a/Engine/Extras/RoboMerge/v3/src/robo/nodebot.ts b/Engine/Extras/RoboMerge/v3/src/robo/nodebot.ts
--- a/Engine/Extras/RoboMerge/v3/src/robo/nodebot.ts
+++ b/Engine/Extras/RoboMerge/v3/src/robo/nodebot.ts
@@ -1627,7 +1627,7 @@ export class NodeBot extends PerforceStatefulBot implements NodeBotInterface {
 			for (const action of info.targets) {
 				if (action.mergeMode === 'normal' && !action.flags.has('disregardassetblock')) {
 
-					for (const blockTargetName of this.branch.blockAssetTargets) {
+					for (const [blockTargetName, blockType] of this.branch.blockAssetTargets) {
 						if (this._getBranch(blockTargetName) === action.branch) {
 
 							const edge = this.getImmediateEdge(action.branch)
@@ -1635,7 +1635,7 @@ export class NodeBot extends PerforceStatefulBot implements NodeBotInterface {
 								throw new Error('Missing edge! ' + action.branch.name)
 							}
 
-							if (info.userRequest || info.cl > edge.lastCl) {
+							if (info.userRequest || (blockType == 'all' && info.cl > edge.lastCl)) {
 								blockAssetEdges.push([edge, action])
 							}
 							break
