Post startup failures to alerts channel

diff --git a/Engine/Extras/RoboMerge/v3/src/robo/graphbot.ts b/Engine/Extras/RoboMerge/v3/src/robo/graphbot.ts
--- a/Engine/Extras/RoboMerge/v3/src/robo/graphbot.ts
+++ b/Engine/Extras/RoboMerge/v3/src/robo/graphbot.ts
@@ -66,7 +66,9 @@ export class GraphBot implements BotEventHandler {
 
 		const result = await BranchDefs.parseAndValidate(graphBot.botLogger, fileText)
 		if (!result.branchGraphDef) {
-			throw new Error(result.validationErrors.length === 0 ? 'Failed to parse' : result.validationErrors.join('\n'))
+			const parseError = result.validationErrors.length === 0 ? 'Failed to parse' : result.validationErrors.join('\n')
+			await postToRobomergeAlerts(`@here Failed to create bot ${botname.toUpperCase}. Halting service startup.\n${parseError}`)
+			throw new Error(parseError)
 		}
 
 		graphBot.branchGraph.config = result.config
@@ -83,13 +85,11 @@ export class GraphBot implements BotEventHandler {
 			error = exc.toString();
 		}
 
-		if (!result.branchGraphDef) {
-			error = result.validationErrors.length === 0 ? 'Failed to parse' : result.validationErrors.join('\n');
-		}
-
 		// start empty bot on error - can be fixed up by branch definition check-in
 		if (error) {
-			graphBot.botLogger.error(`Problem starting up bot ${botname}: ${error}`);
+			const msg = `Problem starting up bot ${botname.toUpperCase}: ${error}`
+			graphBot.botLogger.error(msg);
+			await postToRobomergeAlerts(`@here ${msg}`)
 		}
 
 		graphBot._settings = await Settings.CreateAsync(botname, graphBot.botLogger, graphBot.p4)
