Don't crash if no changes in branch spec root yet, just wait

diff --git a/Engine/Extras/RoboMerge/v3/src/robo/autobranchupdater.ts b/Engine/Extras/RoboMerge/v3/src/robo/autobranchupdater.ts
--- a/Engine/Extras/RoboMerge/v3/src/robo/autobranchupdater.ts
+++ b/Engine/Extras/RoboMerge/v3/src/robo/autobranchupdater.ts
@@ -2,6 +2,7 @@
 
 import * as fs from 'fs';
 import * as path from 'path';
+import { _setTimeout } from '../common/helper';
 import { ContextualLogger } from '../common/logger';
 import { P4_FORCE, PerforceContext, Workspace } from '../common/perforce';
 import { Bot } from './bot-interfaces';
@@ -64,6 +65,22 @@ export class AutoBranchUpdater implements Bot {
 
 		logger.info(`Finding most recent branch spec changelist from ${config.rootPath}`)
 		var bsRoot = config.rootPath + '/...'
+
+		while (true) {
+			try {
+				const change = await p4.latestChange(bsRoot)
+				if (change) {
+					break
+				}
+			}
+			catch (err) {
+			}
+
+			const timeout = 5.0;
+			logger.info(`Will check for changes in ${bsRoot} again in ${timeout} sec...`);
+			await _setTimeout(timeout*1000);
+		}
+		
 		const change = await p4.latestChange(bsRoot)
 		if (change === null)
 			throw new Error(`Unable to query for most recent branch specs CL`)
