Fix track changes failing to track when the first file in the changelist is move/delete and files are in multiple paths that have a common prefix.

diff --git a/Engine/Extras/RoboMerge/v3/src/robo/trackchange.ts b/Engine/Extras/RoboMerge/v3/src/robo/trackchange.ts
--- a/Engine/Extras/RoboMerge/v3/src/robo/trackchange.ts
+++ b/Engine/Extras/RoboMerge/v3/src/robo/trackchange.ts
@@ -122,7 +122,7 @@ export async function trackChange(robo: RoboMerge, logger: ContextualLogger, ini
             }
         }
         else {
-            desc = await p4.describe(clToConsider.cl, 100)
+            desc.entries = (await p4.describe(clToConsider.cl, 100)).entries
         }
     }
 
