Fix failure details not posting a file response to the DM
Do not put the failure details in the message as well

diff --git a/Engine/Extras/RoboMerge/v3/src/robo/nodebot.ts b/Engine/Extras/RoboMerge/v3/src/robo/nodebot.ts
--- a/Engine/Extras/RoboMerge/v3/src/robo/nodebot.ts
+++ b/Engine/Extras/RoboMerge/v3/src/robo/nodebot.ts
@@ -2131,11 +2131,7 @@ export class NodeBot extends PerforceStatefulBot implements NodeBotInterface {
 
 			let message = `While reconsidering CL#${pending.change.cl}`
 
-			if (failure.details) {
-				message += '\n' + failure.details 
-			}
-			else
-			{
+			if (!failure.details) {
 				message += '\n' + failure.description
 			}
 
@@ -2159,7 +2155,7 @@ export class NodeBot extends PerforceStatefulBot implements NodeBotInterface {
 				}
 				
 				let emailAddress = await this.findEmail(owner)
-				this.slackMessages.postDM(emailAddress, pending.newCl, this.branch, dm)
+				await this.slackMessages.postDM(emailAddress, pending.newCl, this.branch, dm)
 
 				if (failure.details) {
 					let file: SlackFile = {
diff --git a/Engine/Extras/RoboMerge/v3/src/robo/notifications.ts b/Engine/Extras/RoboMerge/v3/src/robo/notifications.ts
--- a/Engine/Extras/RoboMerge/v3/src/robo/notifications.ts
+++ b/Engine/Extras/RoboMerge/v3/src/robo/notifications.ts
@@ -320,7 +320,7 @@ export class SlackMessages {
 			dm.channel = channelId
 			// Add the channel/conversation ID to the messageOpts and proceed normally.
 			this.smLogger.info(`Creating direct message for ${emailAddress} (key: ${PersistedSlackMessages.generateKey(cl, branchArg, channelId)})`)
-			this.postOrUpdate(cl, branchArg, dm, persistMessage)
+			await this.postOrUpdate(cl, branchArg, dm, persistMessage)
 		}
 	}
 
@@ -331,7 +331,7 @@ export class SlackMessages {
 			return
 		}
 		file.channels = channelId
-		this.postFile(cl, branchArg, file)
+		await this.postFile(cl, branchArg, file)
 	}
 
 	findAll(cl: number, branchArg: BranchArg) {
