Don't report IPC errors when running functional tests to appease Horde

diff --git a/Engine/Extras/RoboMerge/v3/src/robo/ipc.ts b/Engine/Extras/RoboMerge/v3/src/robo/ipc.ts
--- a/Engine/Extras/RoboMerge/v3/src/robo/ipc.ts
+++ b/Engine/Extras/RoboMerge/v3/src/robo/ipc.ts
@@ -8,6 +8,7 @@ import { OperationResult } from './branch-interfaces';
 import { GraphBot } from './graphbot';
 import { RoboMerge } from './robo';
 import { roboAnalytics } from './roboanalytics';
+import { RoboArgs } from './roboargs';
 import { Status } from './status';
 import { trackChange } from './trackchange';
 import { getPreview } from './preview'
@@ -372,6 +373,14 @@ export class IPC {
 		return { statusCode: 200, message: `${level.toLowerCase()} logging enabled (was: ${previousLevel}).`}
 	}
 
+	private logError(msg: string) {
+		if (RoboArgs.args.runningFunctionalTests) {
+			this.ipcLogger.info(msg)
+		} else {
+			this.ipcLogger.error(msg)
+		}	
+	}
+
 	private async doOperation(botname: string, nodeName: string, operation: string, query: Query, edgeName?: string)
 		: Promise<OperationReturnType> {
 		if (!query.who) {
@@ -558,7 +567,7 @@ export class IPC {
 				}
 			}
 
-			this.ipcLogger.error('Error verifying stomp: ' + stompVerification.message)
+			this.logError('Error verifying stomp: ' + stompVerification.message)
 			return { statusCode: 500, message: stompVerification.message }
 
 		case 'stompchanges':
@@ -577,7 +586,7 @@ export class IPC {
 			if (operationResult.success) {
 				return { statusCode: 200, message: operationResult.message }
 			}
-			this.ipcLogger.error('Error processing stomp: ' + operationResult.message)
+			this.logError('Error processing stomp: ' + operationResult.message)
 			return { statusCode: 500, message: operationResult.message }
 
 		// Requires: cl, target
@@ -605,7 +614,7 @@ export class IPC {
 				}
 			}
 
-			this.ipcLogger.error('Error verifying unlock: ' + unlockVerification.message)
+			this.logError('Error verifying unlock: ' + unlockVerification.message)
 			return { statusCode: 500, message: unlockVerification.message }
 
 		case 'unlockchanges':
@@ -624,7 +633,7 @@ export class IPC {
 			if (operationResult.success) {
 				return { statusCode: 200, message: operationResult.message }
 			}
-			this.ipcLogger.error('Error processing unlock: ' + operationResult.message)
+			this.logError('Error processing unlock: ' + operationResult.message)
 			return { statusCode: 500, message: operationResult.message }
 
 		case 'bypassgatewindow':
diff --git a/Engine/Extras/RoboMerge/v3/src/robo/robo.ts b/Engine/Extras/RoboMerge/v3/src/robo/robo.ts
--- a/Engine/Extras/RoboMerge/v3/src/robo/robo.ts
+++ b/Engine/Extras/RoboMerge/v3/src/robo/robo.ts
@@ -155,6 +155,13 @@ const COMMAND_LINE_ARGS: {[param: string]: Arg<any>} = {
 		dflt: 100
 	},
 
+	runningFunctionalTests: {
+		match: /^(-runningFunctionalTests)$/,
+		parse: str => str === "false" ? false : true,
+		env: 'ROBO_RUNNING_FUNCTIONAL_TESTS',
+		dflt: false
+	},
+
 	helpPageURL: {
 		match: /^-helpPageURL=(.+)$/,
 		env: 'HELP_PAGE_URL',
