Fix stomping on a block assets reconsider edge

diff --git a/Engine/Extras/RoboMerge/v3/src/robo/nodebot.ts b/Engine/Extras/RoboMerge/v3/src/robo/nodebot.ts
--- a/Engine/Extras/RoboMerge/v3/src/robo/nodebot.ts
+++ b/Engine/Extras/RoboMerge/v3/src/robo/nodebot.ts
@@ -1635,7 +1635,7 @@ export class NodeBot extends PerforceStatefulBot implements NodeBotInterface {
 								throw new Error('Missing edge! ' + action.branch.name)
 							}
 
-							if (info.userRequest || (blockType == 'all' && info.cl > edge.lastCl)) {
+							if ((info.userRequest && (!info.forceStompChanges || blockType == 'all')) || (blockType == 'all' && info.cl > edge.lastCl)) {
 								blockAssetEdges.push([edge, action])
 							}
 							break
