Fix display of block assets on the graph

diff --git a/Engine/Extras/RoboMerge/v3/public/js/flow.graph.js b/Engine/Extras/RoboMerge/v3/public/js/flow.graph.js
--- a/Engine/Extras/RoboMerge/v3/public/js/flow.graph.js
+++ b/Engine/Extras/RoboMerge/v3/public/js/flow.graph.js
@@ -18,7 +18,8 @@ class Info {
         this.filteredEdges = new Set([...this.filteredEdges, ...branch.def.graphFilteredTargets.map(edgeId)]);
         this.forcedEdges = new Set([...this.forcedEdges, ...branch.def.forceFlowTo.map(edgeId)]);
         this.defaultDests = new Set([...this.defaultDests, ...branch.def.defaultFlow.map(edgeId)]);
-        this.blockAssetDests = new Set([...this.blockAssetDests, ...branch.def.blockAssetTargets.map(edgeId)]);
+        this.blockAssetDests = new Set([...this.blockAssetDests, 
+            ...branch.def.blockAssetTargets.filter(([key,value]) => value == 'all').map(([key]) => edgeId(key))]);
     }
 }
 
