Fix creating shelf on block asset reconsider edges

diff --git a/Engine/Extras/RoboMerge/v3/src/robo/nodebot.ts b/Engine/Extras/RoboMerge/v3/src/robo/nodebot.ts
--- a/Engine/Extras/RoboMerge/v3/src/robo/nodebot.ts
+++ b/Engine/Extras/RoboMerge/v3/src/robo/nodebot.ts
@@ -1635,7 +1635,8 @@ export class NodeBot extends PerforceStatefulBot implements NodeBotInterface {
 								throw new Error('Missing edge! ' + action.branch.name)
 							}
 
-							if ((info.userRequest && (!info.forceStompChanges || blockType == 'all')) || (blockType == 'all' && info.cl > edge.lastCl)) {
+							const autoFlow = (blockType == 'all' && info.cl > edge.lastCl)
+							if (autoFlow || (info.userRequest && (blockType == 'all' || (!info.forceStompChanges && !info.forceCreateAShelf)))) {
 								blockAssetEdges.push([edge, action])
 							}
 							break
