Fix crash when a change crosses depot (and maybe stream?) boundaries as a result of the use of import+

diff --git a/Engine/Extras/RoboMerge/v3/src/robo/trackchange.ts b/Engine/Extras/RoboMerge/v3/src/robo/trackchange.ts
--- a/Engine/Extras/RoboMerge/v3/src/robo/trackchange.ts
+++ b/Engine/Extras/RoboMerge/v3/src/robo/trackchange.ts
@@ -104,7 +104,7 @@ export async function trackChange(robo: RoboMerge, logger: ContextualLogger, ini
             const rootDir = `${desc.path.slice(0,-3)}*@=${clToConsider.cl}`
             const FILES_PER_DIR = 5
             const DIRS_PER_SLICE = 5
-            const filesPromises = [p4.files(rootDir, FILES_PER_DIR)]
+            const filesPromises = [p4.files(rootDir, FILES_PER_DIR).catch(() => { return []})]
             const dirs = await p4.dirs(rootDir)
             let fpIndex = 0
             let dirSliceBegin = 0
