Sunset messages that are missing their timestamp

diff --git a/Engine/Extras/RoboMerge/v3/src/robo/persistedmessages.ts b/Engine/Extras/RoboMerge/v3/src/robo/persistedmessages.ts
--- a/Engine/Extras/RoboMerge/v3/src/robo/persistedmessages.ts
+++ b/Engine/Extras/RoboMerge/v3/src/robo/persistedmessages.ts
@@ -197,7 +197,12 @@ class DBPersistedSlackMessages extends PersistedSlackMessages {
 
 export async function doDBPersistedMessagesHousekeeping() {
 	const expiredTime = (Date.now() / 1000) - (SUNSET_PERSISTED_NOTIFICATIONS_DAYS * 60 * 60 * 24)
-	return SlackMessagesModel.collection.deleteMany({timestamp_forSunset: { $lt: expiredTime }})
+	return SlackMessagesModel.collection.deleteMany({
+		$or: [
+			{ timestamp_forSunset: { $lt: expiredTime } },
+			{ timestamp_forSunset: NaN }
+		]
+	})
 }
 
 export async function migrateMessagesToDB(botSettings: Settings) {
@@ -287,7 +292,7 @@ class FilePersistedSlackMessages extends PersistedSlackMessages {
 		const keys = Object.keys(persistedMessages)
 		for (const key of keys) {
 			const message = persistedMessages[key]
-			const messageTime = parseFloat(message.timestamp)
+			const messageTime = parseFloat(message.timestamp) || 0 // if message.timestamp doesn't parse, treat as 0 and sunset
 
 			// note: serialisation was briefly wrong (11/11/2019), leaving some message records without a message
 			if (!message.messageOpts || messageTime < expiredTime) {
