Avoid syntax errors triggered by evaluating a specific edge getting in to a loop where they refind and resolve themselves because the node lastCL has moved beyond the syntax error point.

diff --git a/Engine/Extras/RoboMerge/v3/src/robo/conflicts.ts b/Engine/Extras/RoboMerge/v3/src/robo/conflicts.ts
--- a/Engine/Extras/RoboMerge/v3/src/robo/conflicts.ts
+++ b/Engine/Extras/RoboMerge/v3/src/robo/conflicts.ts
@@ -250,6 +250,10 @@ export class Conflicts {
 					remainingConflicts.push(conflict)
 					continue
 				}
+				else if (conflict.kind == 'Syntax error' && [...edges.values()].some(edgeBot => edgeBot.lastCl < conflict.cl)) {
+					remainingConflicts.push(conflict)
+					continue
+				}
 			}
 			else {
 				// If our bot CL is below the conflict CL, we can assume we haven't resolved this conflict
