Limit how large the describeResult for stomping can be to avoid overwhelming persistent data storage

diff --git a/Engine/Extras/RoboMerge/v3/public/js/boilerplate.js b/Engine/Extras/RoboMerge/v3/public/js/boilerplate.js
--- a/Engine/Extras/RoboMerge/v3/public/js/boilerplate.js
+++ b/Engine/Extras/RoboMerge/v3/public/js/boilerplate.js
@@ -1469,7 +1469,11 @@ function renderActionsCell_Common(actionCell, data, operationFunction, operation
 			const stompRequest = '/op/stomp?' + queryParams + location.hash
 			const stompOption = createActionOption('Stomp Changes using ' + conflict.cl + toTargetText, function() {
 				window.location.href = stompRequest;
-			}, `Use ${conflict.cl} to stomp binary changes in ${conflict.target}`);
+			}, conflict.disableStompReason ? conflict.disableStompReason : `Use ${conflict.cl} to stomp binary changes in ${conflict.target}`);
+			if (conflict.disableStompReason) {
+				stompOption.addClass("disabled")
+				stompOption.off('click')
+			}
 
 			// Unlock files
 			const unlockFilesRequest = '/op/unlock?' + queryParams
diff --git a/Engine/Extras/RoboMerge/v3/src/robo/conflict-interfaces.ts b/Engine/Extras/RoboMerge/v3/src/robo/conflict-interfaces.ts
--- a/Engine/Extras/RoboMerge/v3/src/robo/conflict-interfaces.ts
+++ b/Engine/Extras/RoboMerge/v3/src/robo/conflict-interfaces.ts
@@ -21,6 +21,7 @@ export interface PersistentConflict {
 	owner: string
 
 	kind: FailureKind
+	disableStompReason?: string
 	time: Date
 
 	nagCount: number
diff --git a/Engine/Extras/RoboMerge/v3/src/robo/conflicts.ts b/Engine/Extras/RoboMerge/v3/src/robo/conflicts.ts
--- a/Engine/Extras/RoboMerge/v3/src/robo/conflicts.ts
+++ b/Engine/Extras/RoboMerge/v3/src/robo/conflicts.ts
@@ -35,6 +35,7 @@ function makePersistentConflict(blockage: Blockage): PersistentConflict {
 		author: sourceChange.author,
 		owner: blockage.owner,
 		kind: blockage.failure.kind,
+		disableStompReason: blockage.failure.additionalInfo?.disableStompReason,
 		time: blockage.time,
 		nagCount: 0,
 		ugsIssue: -1
@@ -428,6 +429,7 @@ export class Conflicts {
 			}
 
 			conflictObj.kind = conflict.kind
+			conflictObj.disableStompReason = conflict.disableStompReason
 			conflictObj.author = conflict.author
 			conflictObj.owner = conflict.owner
 
diff --git a/Engine/Extras/RoboMerge/v3/src/robo/edgebot.ts b/Engine/Extras/RoboMerge/v3/src/robo/edgebot.ts
--- a/Engine/Extras/RoboMerge/v3/src/robo/edgebot.ts
+++ b/Engine/Extras/RoboMerge/v3/src/robo/edgebot.ts
@@ -795,6 +795,10 @@ class EdgeBotImpl extends PerforceStatefulBot {
 					client: pending.change.targetWorkspace
 				}
 			}
+			if (failure.additionalInfo.describeResult.entries.length > RoboArgs.getMaxStompFileCount()) {
+				failure.additionalInfo.describeResult = null
+				failure.additionalInfo.disableStompReason = "Changelist too large"
+			}
 
 			const conflicts = this.analyzeConflict(result.getConflicts())
 			if (conflicts.length > 0) {
diff --git a/Engine/Extras/RoboMerge/v3/src/robo/nodebot.ts b/Engine/Extras/RoboMerge/v3/src/robo/nodebot.ts
--- a/Engine/Extras/RoboMerge/v3/src/robo/nodebot.ts
+++ b/Engine/Extras/RoboMerge/v3/src/robo/nodebot.ts
@@ -692,6 +692,11 @@ export class NodeBot extends PerforceStatefulBot implements NodeBotInterface {
 	async _evaluateConflictForStomp(changeCl: number, targetBranch: Readonly<EditableBranch>, mergeConflictInfo: MergeConflictAdditionalInformation) : Promise<StompVerification>
 	{
 		try {
+
+			if (!mergeConflictInfo.describeResult) {
+				return this.stompFailureResult(mergeConflictInfo.disableStompReason || "DescribeResult missing for unknown reason")
+			}
+
 			// Get what is left to resolve out of the changelist
 			const resolveResult = Object.assign(new ResolveResultAdditionalInfo(), mergeConflictInfo.resolveResult)
 
diff --git a/Engine/Extras/RoboMerge/v3/src/robo/robo.ts b/Engine/Extras/RoboMerge/v3/src/robo/robo.ts
--- a/Engine/Extras/RoboMerge/v3/src/robo/robo.ts
+++ b/Engine/Extras/RoboMerge/v3/src/robo/robo.ts
@@ -155,6 +155,15 @@ const COMMAND_LINE_ARGS: {[param: string]: Arg<any>} = {
 		dflt: 100
 	},
 
+	maxStompFileCount: {
+		match: /^-maxStompFileCount=([0-9]+)$/,
+		env: 'ROBO_MAX_STOMP_FILE_COUNT',
+		parse: (str: string) => {
+			return parseInt(str)
+		},
+		dflt: 30000
+	},
+
 	runningFunctionalTests: {
 		match: /^(-runningFunctionalTests)$/,
 		parse: str => str === "false" ? false : true,
diff --git a/Engine/Extras/RoboMerge/v3/src/robo/roboargs.ts b/Engine/Extras/RoboMerge/v3/src/robo/roboargs.ts
--- a/Engine/Extras/RoboMerge/v3/src/robo/roboargs.ts
+++ b/Engine/Extras/RoboMerge/v3/src/robo/roboargs.ts
@@ -14,6 +14,10 @@ export class RoboArgs {
 		return this._args.exclusiveLockOpenedsToRun
 	}
 
+	static getMaxStompFileCount() {
+		return this._args.maxStompFileCount
+	}
+
 	static useMongo() {
 		return this._args.mongoDB_URI.length > 0
 	}
diff --git a/Engine/Extras/RoboMerge/v3/src/robo/status-types.ts b/Engine/Extras/RoboMerge/v3/src/robo/status-types.ts
--- a/Engine/Extras/RoboMerge/v3/src/robo/status-types.ts
+++ b/Engine/Extras/RoboMerge/v3/src/robo/status-types.ts
@@ -171,6 +171,7 @@ export type ConflictStatusFields = {
 	target?: string
 	targetStream?: string 
 	kind: FailureKind
+	disableStompReason?: string
 	author: string
 	owner: string
 	slackLinks: string[]
