Missed file for describe result size limitation

diff --git a/Engine/Extras/RoboMerge/v3/src/robo/branch-interfaces.ts b/Engine/Extras/RoboMerge/v3/src/robo/branch-interfaces.ts
--- a/Engine/Extras/RoboMerge/v3/src/robo/branch-interfaces.ts
+++ b/Engine/Extras/RoboMerge/v3/src/robo/branch-interfaces.ts
@@ -252,6 +252,7 @@ export class ResolveResultAdditionalInfo {
 export interface MergeConflictAdditionalInformation {
 	resolveResult: ResolveResultAdditionalInfo,
 	describeResult: DescribeResult,
+	disableStompReason?: string,
 	client?: string
 }
 
