Handle oversized DB save failures
Clean out oversized describe results on load

diff --git a/Engine/Extras/RoboMerge/v3/src/robo/branch-interfaces.ts b/Engine/Extras/RoboMerge/v3/src/robo/branch-interfaces.ts
--- a/Engine/Extras/RoboMerge/v3/src/robo/branch-interfaces.ts
+++ b/Engine/Extras/RoboMerge/v3/src/robo/branch-interfaces.ts
@@ -251,11 +251,22 @@ export class ResolveResultAdditionalInfo {
 
 export interface MergeConflictAdditionalInformation {
 	resolveResult: ResolveResultAdditionalInfo,
-	describeResult: DescribeResult,
+	describeResult?: DescribeResult,
 	disableStompReason?: string,
 	client?: string
 }
 
+export function isA_MergeConflictAdditionalInformation(obj: any): obj is MergeConflictAdditionalInformation {
+  return (
+    obj !== null &&
+    typeof obj === "object" &&
+    typeof obj.resolveResult === "object" &&
+    (obj.describeResult === undefined || typeof obj.describeResult === "object") &&
+    (obj.disableStompReason === undefined || typeof obj.disableStompReason === "string") &&
+    (obj.client === undefined || typeof obj.client === "string")
+  )
+}
+
 export interface ForcedCl {
 	nodeOrEdgeName: string
 	sourceBranchUpperName: string
diff --git a/Engine/Extras/RoboMerge/v3/src/robo/edgebot.ts b/Engine/Extras/RoboMerge/v3/src/robo/edgebot.ts
--- a/Engine/Extras/RoboMerge/v3/src/robo/edgebot.ts
+++ b/Engine/Extras/RoboMerge/v3/src/robo/edgebot.ts
@@ -796,7 +796,7 @@ class EdgeBotImpl extends PerforceStatefulBot {
 				}
 			}
 			if (failure.additionalInfo.describeResult.entries.length > RoboArgs.getMaxStompFileCount()) {
-				failure.additionalInfo.describeResult = null
+				failure.additionalInfo.describeResult = undefined
 				failure.additionalInfo.disableStompReason = "Changelist too large"
 			}
 
diff --git a/Engine/Extras/RoboMerge/v3/src/robo/settings.ts b/Engine/Extras/RoboMerge/v3/src/robo/settings.ts
--- a/Engine/Extras/RoboMerge/v3/src/robo/settings.ts
+++ b/Engine/Extras/RoboMerge/v3/src/robo/settings.ts
@@ -3,11 +3,13 @@
 import * as fs from 'fs';
 import { Args } from '../common/args';
 import { ContextualLogger } from '../common/logger';
-import { VersionReader } from '../common/version';
 import { PerforceContext, Workspace } from '../common/perforce';
+import { VersionReader } from '../common/version';
+import { isA_MergeConflictAdditionalInformation, MergeConflictAdditionalInformation } from './branch-interfaces';
+import { postToRobomergeAlerts } from './notifications';
+import { migrateMessagesToDB } from './persistedmessages';
 import { RoboArgs } from './roboargs';
 import semver = require('semver');
-import { migrateMessagesToDB } from './persistedmessages';
 
 const BotModel = require('./models/bot');
 
@@ -181,6 +183,11 @@ export class Settings {
 			throw new Error(`Settings files prior to 3.0.2 are unsupported`)
 		}
 
+		// Clean up oversized describe records bloating files
+		if (semver.lt(objVersion, '3.3.0')) {
+			settings._shrink(true)
+		}
+
 		/**
 		 * MIGRATION CODE SECTION START
 		 */
@@ -229,6 +236,20 @@ export class Settings {
 							this.savingState = 'Idle'
 						}
 					})
+					.catch(() => {
+						const size = Buffer.byteLength(JSON.stringify(this.object));
+						this.settingsLogger.warn(`Unable to save ${this.botname}. Size: ${size}. Attempting to shrink and resave.`)
+						this._shrink()
+						const shrunkSize = Buffer.byteLength(JSON.stringify(this.object));
+						if (shrunkSize < size) {
+							saveToDB()
+						}
+						else {
+							const errorMsg = `Unable to save ${this.botname}. Size: ${size}.`
+							this.settingsLogger.error(errorMsg)
+							postToRobomergeAlerts(errorMsg)
+						}
+					})
 				}
 				if (this.savingState == 'Idle') {
 					saveToDB()
@@ -279,5 +300,29 @@ export class Settings {
 			}
 		}
 	}
+
+	_shrink(trim?: boolean) {
+			Object.keys(this.object).forEach(key => {
+				if (key !== 'version') {
+					this._shrinkNode(key, trim)
+				}
+			})
+	}
+
+	_shrinkNode(nodeName: string, trim?: boolean) {
+		const edges = this.object[nodeName]?.edges || {}
+		for (const key of Object.keys(edges)) {
+			const edge = this.object[nodeName].edges[key]
+			const additionalInfo: MergeConflictAdditionalInformation = edge.pause?.blockage?.additionalInfo
+			if (isA_MergeConflictAdditionalInformation(additionalInfo)) {
+				if (additionalInfo?.describeResult) {
+					if (!trim || additionalInfo.describeResult.entries.length > RoboArgs.getMaxStompFileCount()) {
+						additionalInfo.disableStompReason = "Changelist too large"
+						additionalInfo.describeResult = undefined
+					}
+				}
+			}
+		}
+	}
 }
 
