Save details per node within bot

diff --git a/Engine/Extras/RoboMerge/v3/src/robo/models/bot.ts b/Engine/Extras/RoboMerge/v3/src/robo/models/bot.ts
--- a/Engine/Extras/RoboMerge/v3/src/robo/models/bot.ts
+++ b/Engine/Extras/RoboMerge/v3/src/robo/models/bot.ts
@@ -7,6 +7,7 @@ var mongoose = require('mongoose');
  */
 var botSchema = new mongoose.Schema({
   botname: { type: String, index: true },
+  nodename: { type: String, index: true },
   details: {  }
 });
 
diff --git a/Engine/Extras/RoboMerge/v3/src/robo/settings.ts b/Engine/Extras/RoboMerge/v3/src/robo/settings.ts
--- a/Engine/Extras/RoboMerge/v3/src/robo/settings.ts
+++ b/Engine/Extras/RoboMerge/v3/src/robo/settings.ts
@@ -30,6 +30,10 @@ export function settingsInit(inArgs: Args) {
 	args = inArgs
 }
 	
+interface SaveObjectOpts {
+	name?: string
+	ensureInP4?: boolean
+}
 
 const jsonlint: any = require('jsonlint-mod')
 
@@ -84,7 +88,7 @@ export class Context {
 	}
 	set(name: string, value: any) {
 		this.object[name] = value;
-		this.settings._saveObject();
+		this.settings._saveObject({name: this.path[0]});
 	}
 }
 
@@ -103,11 +107,13 @@ function getPersistenceFilepath(botname: string) {
 	return args.persistenceDir + `/${botname}.settings`
 }
 
+type SavingState = 'Saving' | 'NeedsSave'
+
 export class Settings {
 	botname: string
 	enableSave: boolean
 	object: { [key: string]: any }
-	private savingState: 'Idle' | 'Saving' | 'NeedsSave' = 'Idle'
+	private savingState: Map<string, SavingState> = new Map()
 	private	lastPersistBackupTime: Date
 	private readonly settingsLogger: ContextualLogger
 	private readonly p4backupWorkspace?: Workspace
@@ -134,13 +140,26 @@ export class Settings {
 
 		if (RoboArgs.useMongo()) {
 			// first try to get the bot data from the DB
-			const botRecord = await BotModel.find({botname: botname})
-			if (botRecord.length > 1) {
-				throw new Error(`Found multiple records for ${botname}`) 
-			}
-
-			if (botRecord.length == 1) {
-				settings.object = botRecord[0].details
+			const botRecords = await BotModel.find({botname: botname})
+			if (botRecords.length > 1) {
+				// starting with 3.3.0 the individual node records themselves are being persisted/versioned, 
+				// not the top level object itself, so we can set the object to that version to indicate there
+				// should only be versioning on the individual nodes
+				settings.object = {version: '3.3.0'} 
+				
+				for (const botRecord of botRecords) {
+					settings.object[botRecord.nodename] = botRecord.details
+				}
+			}			
+			else if (botRecords.length == 1) {
+				const objVersion = semver.coerce(String(botRecords[0].details.version))!
+				if (semver.lt(objVersion, '3.3.0')) {
+					settings.object = botRecords[0].details
+					BotModel.deleteOne({_id: botRecords[0]._id})
+				}
+				else {
+					settings.object[botRecords[0].nodename] = botRecords[0].details
+				}
 			}
 		}
 
@@ -209,7 +228,7 @@ export class Settings {
 		 */
 		
 		settings.object.version = VERSION.raw
-		settings._saveObject(true);
+		settings._saveObject({ensureInP4: true});
 
 		return settings
 	}
@@ -222,40 +241,50 @@ export class Settings {
 		return getPersistenceFilepath(this.botname) + FilePurposeExtensionMap[purpose]
 	}
 
-	_saveObject(ensureInP4?: boolean) {
+	_saveObject(opts?: SaveObjectOpts) {
 		if (this.enableSave) {
 			if (RoboArgs.useMongo()) {
-				const saveToDB = () => {
-					this.savingState = 'Saving'
-					BotModel.collection.updateOne({botname:this.botname},{$set:{botname:this.botname,details:this.object}},{upsert:true})
-					.then(() => {
-						if (this.savingState == 'NeedsSave') {
-							saveToDB()
-						}
-						else {
-							this.savingState = 'Idle'
-						}
-					})
-					.catch(() => {
-						const size = Buffer.byteLength(JSON.stringify(this.object));
-						this.settingsLogger.warn(`Unable to save ${this.botname}. Size: ${size}. Attempting to shrink and resave.`)
-						this._shrink()
-						const shrunkSize = Buffer.byteLength(JSON.stringify(this.object));
-						if (shrunkSize < size) {
-							saveToDB()
-						}
-						else {
-							const errorMsg = `Unable to save ${this.botname}. Size: ${size}.`
-							this.settingsLogger.error(errorMsg)
-							postToRobomergeAlerts(errorMsg)
-						}
-					})
-				}
-				if (this.savingState == 'Idle') {
-					saveToDB()
+				if (opts?.name) {
+					const saveToDB = () => {
+						this.savingState.set(opts.name!, 'Saving')
+						this.object[opts.name!].version = this.object.version
+						BotModel.collection.updateOne({botname:this.botname,nodename:opts.name},{$set:{botname:this.botname,nodename:opts.name,details:this.object[opts.name!]}},{upsert:true})
+						.then(() => {
+							if (this.savingState.get(opts.name!) == 'NeedsSave') {
+								saveToDB()
+							}
+							else {
+								this.savingState.delete(opts.name!)
+							}
+						})
+						.catch(() => {
+							const size = Buffer.byteLength(JSON.stringify(this.object[opts.name!]));
+							this.settingsLogger.warn(`Unable to save ${this.botname}:${opts.name}. Size: ${size}. Attempting to shrink and resave.`)
+							this._shrinkNode(opts.name!)
+							const shrunkSize = Buffer.byteLength(JSON.stringify(this.object[opts.name!]));
+							if (shrunkSize < size) {
+								saveToDB()
+							}
+							else {
+								const errorMsg = `Unable to save ${this.botname}:${opts.name}. Size: ${size}.`
+								this.settingsLogger.error(errorMsg)
+								postToRobomergeAlerts(errorMsg)
+							}
+						})
+					}
+					if (!this.savingState.has(opts.name)) {
+						saveToDB()
+					}
+					else {
+						this.savingState.set(opts.name, 'NeedsSave')
+					}
 				}
 				else {
-					this.savingState = 'NeedsSave'
+					Object.keys(this.object).forEach(key => {
+						if (key !== 'version') {
+							this._saveObject({name: key})
+						}
+					})
 				}
 			}
 			else {
@@ -272,7 +301,7 @@ export class Settings {
 				fs.copyFileSync(tempFile, persistentFile)
 
 				if (this.p4backupWorkspace) {
-					if (ensureInP4) {
+					if (opts?.ensureInP4) {
 						this.p4.fstat(this.p4backupWorkspace, persistentFile)
 						.then((fstat) => {
 							if (fstat.length == 0 || fstat[0].headAction == 'delete') {
