Fix nagging not working when using additional slack channels

diff --git a/Engine/Extras/RoboMerge/v3/src/robo/branch-interfaces.ts b/Engine/Extras/RoboMerge/v3/src/robo/branch-interfaces.ts
--- a/Engine/Extras/RoboMerge/v3/src/robo/branch-interfaces.ts
+++ b/Engine/Extras/RoboMerge/v3/src/robo/branch-interfaces.ts
@@ -6,6 +6,7 @@ import { NodeBotInterface } from './bot-interfaces';
 import { FailureKind } from './status-types';
 import { BotConfig, BranchBase, ApprovalOptions, EdgeOptions, NodeOptions, IntegrationWindowPane } from './branchdefs';
 import { BlockageNodeOpUrls } from './roboserver';
+import { PersistentConflict } from './conflict-interfaces';
 
 export type BranchArg = Branch | string
 export function resolveBranchArg(branchArg: BranchArg, toUpperCase?: boolean): string {
@@ -352,3 +353,10 @@ export type EndIntegratingToGateEvent = {
 	context: GateEventContext
 	targetCl: number
 }
+
+export type NagBlockageEvent = {
+	conflict: PersistentConflict
+	triager: string
+	triagerEmail: Promise<string>
+	ownerEmail: Promise<string>
+}
\ No newline at end of file
diff --git a/Engine/Extras/RoboMerge/v3/src/robo/events.ts b/Engine/Extras/RoboMerge/v3/src/robo/events.ts
--- a/Engine/Extras/RoboMerge/v3/src/robo/events.ts
+++ b/Engine/Extras/RoboMerge/v3/src/robo/events.ts
@@ -1,13 +1,14 @@
 // Copyright Epic Games, Inc. All Rights Reserved.
 
 import { AlreadyIntegrated, Blockage, ChangeInfo, ForcedCl, PendingChange } from './branch-interfaces';
-import { BeginIntegratingToGateEvent, EndIntegratingToGateEvent } from './branch-interfaces';
+import { BeginIntegratingToGateEvent, EndIntegratingToGateEvent, NagBlockageEvent } from './branch-interfaces';
 import { BotConfig } from './branchdefs';
 import { PersistentConflict } from './conflict-interfaces';
 
 type OnChangeParsed = (arg: ChangeInfo) => void
 type OnBlockage = (arg: Blockage, isNew: boolean) => void
 type OnBlockageAcknowledged = (arg: PersistentConflict) => void
+type OnNagBlockage = (arg: NagBlockageEvent) => void
 type OnCommit = (arg: PendingChange) => void
 type OnAlreadyIntegrated = (arg: AlreadyIntegrated) => void
 type OnForcedLastCl = (arg: ForcedCl) => void
@@ -52,6 +53,10 @@ export class BotEvents {
 		this.blockageAcknowledgedListeners.push(listener)
 	}
 
+	onNagBlockage(listener: OnNagBlockage) {
+		this.nagBlockageListeners.push(listener)
+	}
+
 	// Fired when a merge has been committed
 	onCommit(listener: OnCommit) {
 		this.commitListeners.push(listener)
@@ -99,6 +104,7 @@ export class BotEvents {
 				case 'onChangeParsed': this.changeParsedListeners.push(arg => prop.call(handler, arg)); break
 				case 'onBlockage': this.blockageListeners.push((arg, isNew) => prop.call(handler, arg, isNew)); break
 				case 'onBlockageAcknowledged': this.blockageAcknowledgedListeners.push((arg) => prop.call(handler, arg)); break
+				case 'onNagBlockage': this.nagBlockageListeners.push((arg) => prop.call(handler, arg)); break
 				case 'onCommit': this.commitListeners.push(arg => prop.call(handler, arg)); break
 				case 'onAlreadyIntegrated': this.alreadyIntegratedListeners.push(arg => prop.call(handler, arg)); break
 				case 'onForcedLastCl': this.forcedLastClListeners.push(arg => prop.call(handler, arg)); break
@@ -114,6 +120,7 @@ export class BotEvents {
 	protected readonly changeParsedListeners: OnChangeParsed[] = []
 	protected readonly blockageListeners: OnBlockage[] = []
 	protected readonly blockageAcknowledgedListeners: OnBlockageAcknowledged[] = []
+	protected readonly nagBlockageListeners: OnNagBlockage[] = []
 	protected readonly commitListeners: OnCommit[] = []
 	protected readonly alreadyIntegratedListeners: OnAlreadyIntegrated[] = []
 	protected readonly forcedLastClListeners: OnForcedLastCl[] = []
@@ -144,6 +151,12 @@ export class BotEventTriggers extends BotEvents {
 		}
 	}
 
+	nagBlockage(arg: NagBlockageEvent) {
+		for (const listener of this.nagBlockageListeners) {
+			listener(arg)
+		}
+	}
+
 	reportCommit(arg: PendingChange) {
 		for (const listener of this.commitListeners) {
 			listener(arg)
diff --git a/Engine/Extras/RoboMerge/v3/src/robo/nodebot.ts b/Engine/Extras/RoboMerge/v3/src/robo/nodebot.ts
--- a/Engine/Extras/RoboMerge/v3/src/robo/nodebot.ts
+++ b/Engine/Extras/RoboMerge/v3/src/robo/nodebot.ts
@@ -16,7 +16,7 @@ import { Conflicts } from './conflicts';
 import { EdgeBot, EdgeIntegrationDetails, EdgeMergeResults } from './edgebot';
 import { BotEventTriggers } from './events';
 import { GraphBot } from './graphbot';
-import { getPingTarget, isUserSlackGroup, makeClLink, SlackMessages } from './notifications';
+import { makeClLink, SlackMessages } from './notifications';
 import { PerforceStatefulBot } from './perforce-stateful-bot';
 import { RoboArgs } from './roboargs';
 import { BlockageNodeOpUrls, OperationUrlHelper } from './roboserver';
@@ -2285,55 +2285,13 @@ export class NodeBot extends PerforceStatefulBot implements NodeBotInterface {
 						conflict.nagCount++ 
 						conflict.lastNagTime = new Date(now)
 						naggedSomeone = true
-
-						let timeDesc
-						if (ageMinutes < 60) {
-							timeDesc = `${Math.floor(ageMinutes)} minutes`
-						}
-						else {
-							const hours = Math.floor(ageMinutes / 60)
-							if (hours == 1) {
-								timeDesc = "1 hour"
-							}
-							else {
-								timeDesc = `${hours} hours`
-							}
-						}
-				
-						let triager = getNagProperty("triager")
-
-						if (triager && !isUserSlackGroup(triager))
-						{
-							const emailAddress = await this.findEmail(triager)
-							if (emailAddress)
-							{
-								const user = await this.slackMessages.getSlackUser(emailAddress)
-								if (user) {
-									triager = `<@${user}>`
-								}
-								else {
-									triager = `@${triager}`
-									this.nodeBotLogger.error(`Unable to look up triager from ${triager} (${emailAddress})`)
-								}
-							}
-							else {
-								triager = `@${triager}`
-								this.nodeBotLogger.error(`Unable to look up email address for triager ${triager}`)
-							}
-						}
-
-						const channelPing = await getPingTarget(await this.p4.getEmail(conflict.owner), conflict.owner, this.slackMessages)
-
-						const nagMessage: SlackMessage = {
-							text: `${triager ? triager + ' ' : ''}${channelPing} RoboMerge blocked for more than ${timeDesc}`,
-							style: SlackMessageStyles.DANGER,
-							channel: this.branchGraph.config.slackChannel,
-							mrkdwn: true
-						}
-
-						this.nodeBotLogger.info(`Sending nag notification to ${triager ? triager + ' and ' : ''}${conflict.owner} after ${Math.floor(ageMinutes)} minutes`)
-						const branchName =  (conflict.kind === "Syntax error" ? conflict.kind : conflict.targetBranchName || conflict.blockedBranchName)
-						this.slackMessages.postReply(conflict.cl, branchName, nagMessage)
+						const triager = getNagProperty("triager")
+						this.eventTriggers.nagBlockage({
+							conflict, 
+							triager,
+							triagerEmail: triager && this.findEmail(triager),
+							ownerEmail: this.findEmail(conflict.owner)
+						})
 					}
 				}
 				else if (conflict.nagCount == 0) // for now we'll maintain the email only once behavior
diff --git a/Engine/Extras/RoboMerge/v3/src/robo/notifications.ts b/Engine/Extras/RoboMerge/v3/src/robo/notifications.ts
--- a/Engine/Extras/RoboMerge/v3/src/robo/notifications.ts
+++ b/Engine/Extras/RoboMerge/v3/src/robo/notifications.ts
@@ -6,7 +6,7 @@ import { Badge } from '../common/badge';
 import { Random, setDefault } from '../common/helper';
 import { ContextualLogger } from '../common/logger';
 import { ExclusiveFileDetails } from '../common/perforce';
-import { Blockage, Branch, BranchArg,  ExclusiveLockInfo, ForcedCl, MergeAction, NodeOpUrlGenerator, resolveBranchArg } from './branch-interfaces';
+import { Blockage, Branch, BranchArg, EditableBranch, ExclusiveLockInfo, ForcedCl, MergeAction, NagBlockageEvent, NodeOpUrlGenerator, resolveBranchArg } from './branch-interfaces';
 import { PersistentConflict, Resolution } from './conflict-interfaces';
 import { BotEventHandler, BotEvents } from './events';
 import { NodeBot } from './nodebot';
@@ -447,6 +447,36 @@ export class BotNotifications implements BotEventHandler {
 		}
 	}
 
+	getChannelsToPostTo(sourceBranch: Readonly<EditableBranch>|string, targetBranchName?: string) {
+		let channelsToPostTo = []
+		if (targetBranchName) {
+			const sourceBranchName = (typeof(sourceBranch) === 'string' ? sourceBranch : sourceBranch.upperName)
+			const additionalChannelInfo = this.additionalBlockChannelIds.get(sourceBranchName + '|' + targetBranchName)
+			if (additionalChannelInfo) {
+				const [sideChannel, postOnlyToSideChannel] = additionalChannelInfo
+				if (!postOnlyToSideChannel) {
+					channelsToPostTo.push(this.slackChannel)
+				}
+				channelsToPostTo.push(sideChannel)
+			}
+			else {
+				channelsToPostTo.push(this.slackChannel)
+			}
+		}
+		else if (typeof(sourceBranch) !== 'string') {
+			if (sourceBranch.config.additionalSlackChannelForBlockages) {
+				if (!sourceBranch.config.postMessagesToAdditionalChannelOnly) {
+					channelsToPostTo.push(this.slackChannel)
+				}
+				channelsToPostTo.push(sourceBranch.config.additionalSlackChannelForBlockages)
+			}
+			else {
+				channelsToPostTo.push(this.slackChannel)
+			}
+		}
+		return channelsToPostTo
+	}
+
 	/** Conflict */
 	// considered if no Slack set up, should continue to add people to notify, but too complicated:
 	//		let's go all in on Slack. Fine for this to be async (but fire and forget)
@@ -498,7 +528,7 @@ export class BotNotifications implements BotEventHandler {
 		}
 
 		let messagesToPost = []
-		let channelsToPostTo = []
+		const channelsToPostTo = this.getChannelsToPostTo(sourceBranch, targetBranch?.upperName)
 		let usersToInvite: Set<string> = new Set()
 
 		if (userEmail) {
@@ -506,18 +536,6 @@ export class BotNotifications implements BotEventHandler {
 		}
 
 		if (targetBranch) {
-			const additionalChannelInfo = this.additionalBlockChannelIds.get(sourceBranch.upperName + '|' + targetBranch.upperName)
-			if (additionalChannelInfo) {
-				const [sideChannel, postOnlyToSideChannel] = additionalChannelInfo
-				if (!postOnlyToSideChannel) {
-					channelsToPostTo.push(message.channel)
-				}
-				channelsToPostTo.push(sideChannel)
-			}
-			else {
-				channelsToPostTo.push(message.channel)
-			}
-
 			messagesToPost.push({ message })
 
 			if (isNew && blockage.failure.details) {
@@ -565,15 +583,6 @@ export class BotNotifications implements BotEventHandler {
 			}
 		}
 		else {
-			if (sourceBranch.config.additionalSlackChannelForBlockages) {
-				if (!sourceBranch.config.postMessagesToAdditionalChannelOnly) {
-					channelsToPostTo.push(message.channel)
-				}
-				channelsToPostTo.push(sourceBranch.config.additionalSlackChannelForBlockages)
-			}
-			else {
-				channelsToPostTo.push(message.channel)
-			}
 			messagesToPost.push({ message })
 
 			const syntaxErrorMessage: SlackMessage = {
@@ -652,6 +661,65 @@ export class BotNotifications implements BotEventHandler {
 		}
 	}
 
+	async onNagBlockage(evt: NagBlockageEvent) {
+		if (this.slackMessages) {
+			const conflict = evt.conflict
+			const ageMinutes = ((Date.now() - conflict.time.getTime()) / 1000) / 60
+			let timeDesc
+			if (ageMinutes < 60) {
+				timeDesc = `${Math.floor(ageMinutes)} minutes`
+			}
+			else {
+				const hours = Math.floor(ageMinutes / 60)
+				if (hours == 1) {
+					timeDesc = "1 hour"
+				}
+				else {
+					timeDesc = `${hours} hours`
+				}
+			}
+
+			let triager = evt.triager
+			if (triager && !isUserSlackGroup(triager))
+			{
+				const emailAddress = await evt.triagerEmail
+				if (emailAddress)
+				{
+					const user = await this.slackMessages.getSlackUser(emailAddress)
+					if (user) {
+						triager = `<@${user}>`
+					}
+					else {
+						triager = `@${triager}`
+						this.botNotificationsLogger.error(`Unable to look up triager from ${triager} (${emailAddress})`)
+					}
+				}
+				else {
+					triager = `@${triager}`
+					this.botNotificationsLogger.error(`Unable to look up email address for triager ${triager}`)
+				}
+			}
+
+			const channelPing = await getPingTarget(await evt.ownerEmail, conflict.owner, this.slackMessages)
+
+			const nagMessage: SlackMessage = {
+				text: `${triager ? triager + ' ' : ''}${channelPing} RoboMerge blocked for more than ${timeDesc}`,
+				style: SlackMessageStyles.DANGER,
+				channel: this.slackChannel,
+				mrkdwn: true
+			}
+
+			const channelsToPostTo = this.getChannelsToPostTo(conflict.blockedBranchName, conflict.targetBranchName)
+
+			this.botNotificationsLogger.info(`Sending nag notification to ${triager ? triager + ' and ' : ''}${conflict.owner} after ${Math.floor(ageMinutes)} minutes`)
+			const branchName = (conflict.kind === "Syntax error" ? conflict.kind : conflict.targetBranchName || conflict.blockedBranchName)
+			for (let slackChannel of channelsToPostTo) {
+				nagMessage.channel = slackChannel
+				await this.slackMessages.postReply(conflict.cl, branchName, nagMessage)
+			}
+		}
+	}
+
 	////////////////////////
 	// Conflict resolution
 	//
@@ -693,19 +761,7 @@ export class BotNotifications implements BotEventHandler {
 
 	onNonSkipLastClChange(details: ForcedCl) {
 		if (this.slackMessages) {
-			let channelsToPostTo = []		
-			const additionalChannelInfo = this.additionalBlockChannelIds.get(details.sourceBranchUpperName + '|' + details.targetBranchUpperName)
-			if (additionalChannelInfo) {
-				const [sideChannel, postOnlyToSideChannel] = additionalChannelInfo
-				if (!postOnlyToSideChannel) {
-					channelsToPostTo.push(this.slackChannel)
-				}
-				channelsToPostTo.push(sideChannel)
-			}
-			else {
-				channelsToPostTo.push(this.slackChannel)
-			}
-
+			const channelsToPostTo = this.getChannelsToPostTo(details.sourceBranchUpperName, details.targetBranchUpperName)
 			for (let slackChannel of channelsToPostTo) {
 				this.slackMessages.postNonConflictMessage({
 					title: details.nodeOrEdgeName + ' forced to new CL',
