Remove unncessary intermediate function and logging

diff --git a/Engine/Extras/RoboMerge/v3/src/robo/nodebot.ts b/Engine/Extras/RoboMerge/v3/src/robo/nodebot.ts
--- a/Engine/Extras/RoboMerge/v3/src/robo/nodebot.ts
+++ b/Engine/Extras/RoboMerge/v3/src/robo/nodebot.ts
@@ -412,7 +412,7 @@ export class NodeBot extends PerforceStatefulBot implements NodeBotInterface {
 						mrkdwn: true
 					}
 					
-					let emailAddress = await this.findEmail(fromQueue.who)
+					let emailAddress = await this.p4.getEmail(fromQueue.who)
 					this.slackMessages.postDM(emailAddress, fromQueue.cl, this.branch, dm)
 				}
 				else {
@@ -436,7 +436,7 @@ export class NodeBot extends PerforceStatefulBot implements NodeBotInterface {
 						mrkdwn: true
 					}
 					
-					let emailAddress = await this.findEmail(fromQueue.who)
+					let emailAddress = await this.p4.getEmail(fromQueue.who)
 					this.slackMessages.postDM(emailAddress, fromQueue.cl, this.branch, dm)
 				}
 				else {
@@ -475,7 +475,7 @@ export class NodeBot extends PerforceStatefulBot implements NodeBotInterface {
 					mrkdwn: true
 				}
 				
-				let emailAddress = await this.findEmail(fromQueue.who)
+				let emailAddress = await this.p4.getEmail(fromQueue.who)
 				this.slackMessages.postDM(emailAddress, fromQueue.cl, this.branch, dm)
 			}
 
@@ -1060,7 +1060,7 @@ export class NodeBot extends PerforceStatefulBot implements NodeBotInterface {
 			filesByUser.forEach(async (files, author) => {
 
 				let userToNotify = `@${author}`
-				const emailAddress = await this.findEmail(author)
+				const emailAddress = await this.p4.getEmail(author)
 				if (emailAddress)
 				{
 					const user = await this.slackMessages!.getSlackUser(emailAddress)
@@ -1381,7 +1381,7 @@ export class NodeBot extends PerforceStatefulBot implements NodeBotInterface {
 				mrkdwn: true
 			}
 			
-			this.slackMessages.postDM(this.findEmail(event.change.owner!), event.change.cl, this.branch, dm)			
+			this.slackMessages.postDM(this.p4.getEmail(event.change.owner!), event.change.cl, this.branch, dm)			
 		}
 	}
 
@@ -1395,11 +1395,6 @@ export class NodeBot extends PerforceStatefulBot implements NodeBotInterface {
 		}
 	}
 
-	findEmail = (user: string) => {
-		this._log_action(`Getting email for user ${user}`);
-		return this.p4.getEmail(user);
-	}
-
 	// Listener method to email a blockage owner when the branch encounters said blockage,
 	// depending on NodeBot configuration
 	private emailBlockageCulprit(blockage: Blockage) {
@@ -1452,7 +1447,7 @@ export class NodeBot extends PerforceStatefulBot implements NodeBotInterface {
 				mrkdwn: true
 			}
 			
-			let emailAddress = await this.findEmail(owner)
+			let emailAddress = await this.p4.getEmail(owner)
 			this.slackMessages.postDM(emailAddress, pendingChange.newCl, this.branch, dm)
 		}
 		else
@@ -1730,7 +1725,7 @@ export class NodeBot extends PerforceStatefulBot implements NodeBotInterface {
 	private _sendGenericEmail(recipients: Recipients, subject: string, intro: string, message: string) {
 		const params : MailParams = { intro, message, botname: this.graphBotName, 
 			robomergeRootUrl: { value: this.externalRobomergeUrl, noEscape: true } }
-		this.mailer.sendEmail('generic-email-template', recipients, subject, params, this.findEmail)
+		this.mailer.sendEmail('generic-email-template', recipients, subject, params, this.p4.getEmail)
 	}
 	private _sendBlockageEmail(blockage: Blockage) {
 		const recipients = new Recipients(blockage.owner)
@@ -1774,7 +1769,7 @@ export class NodeBot extends PerforceStatefulBot implements NodeBotInterface {
 		// Add list to email parameters
 		params.resolutionOptions = { value: resolutionOptionsStr, noEscape: true }
 
-		this.mailer.sendEmail('new-blockage', recipients, subject, params, this.findEmail)
+		this.mailer.sendEmail('new-blockage', recipients, subject, params, this.p4.getEmail)
 	}
 
 	sendErrorEmail(recipients: Recipients, subject: string, message: string, errorPreface='RoboMerge needs your help:'): void {
@@ -2102,7 +2097,7 @@ export class NodeBot extends PerforceStatefulBot implements NodeBotInterface {
 					mrkdwn: true
 				}
 				
-				let emailAddress = await this.findEmail(owner)
+				let emailAddress = await this.p4.getEmail(owner)
 				this.slackMessages.postDM(emailAddress, change.cl, this.branch, dm)				
 			}
 			else {
@@ -2160,7 +2155,7 @@ export class NodeBot extends PerforceStatefulBot implements NodeBotInterface {
 					mrkdwn: true
 				}
 				
-				let emailAddress = await this.findEmail(owner)
+				let emailAddress = await this.p4.getEmail(owner)
 				await this.slackMessages.postDM(emailAddress, pending.newCl, this.branch, dm)
 
 				if (failure.details) {
@@ -2289,8 +2284,8 @@ export class NodeBot extends PerforceStatefulBot implements NodeBotInterface {
 						this.eventTriggers.nagBlockage({
 							conflict, 
 							triager,
-							triagerEmail: triager && this.findEmail(triager),
-							ownerEmail: this.findEmail(conflict.owner)
+							triagerEmail: triager && this.p4.getEmail(triager),
+							ownerEmail: this.p4.getEmail(conflict.owner)
 						})
 					}
 				}
