When a known bot has a blockage ping the triager immediately if one is set

diff --git a/Engine/Extras/RoboMerge/v3/src/robo/branch-interfaces.ts b/Engine/Extras/RoboMerge/v3/src/robo/branch-interfaces.ts
--- a/Engine/Extras/RoboMerge/v3/src/robo/branch-interfaces.ts
+++ b/Engine/Extras/RoboMerge/v3/src/robo/branch-interfaces.ts
@@ -203,6 +203,8 @@ export interface Blockage {
 	failure: Failure
 	owner: string
 	ownerEmail: Promise<string | null>
+	triager?: string
+	triagerEmail?: Promise<string | null>
 	approval?: {
 		settings: ApprovalOptions
 		shelfCl: number
diff --git a/Engine/Extras/RoboMerge/v3/src/robo/nodebot.ts b/Engine/Extras/RoboMerge/v3/src/robo/nodebot.ts
--- a/Engine/Extras/RoboMerge/v3/src/robo/nodebot.ts
+++ b/Engine/Extras/RoboMerge/v3/src/robo/nodebot.ts
@@ -2110,9 +2110,11 @@ export class NodeBot extends PerforceStatefulBot implements NodeBotInterface {
 			const blockageOccurred = existingBlockage ? existingBlockage.time : new Date
 
 			if (!existingBlockage) {
-				const ownerEmailJustForSlackAtTheMoment = this.p4.getEmail(owner)
+				const ownerEmail = this.p4.getEmail(owner)
+				const triager = this.getBranchProperty("triager")
+				const triagerEmail = triager && this.p4.getEmail(triager)
 				this.conflicts.onBlockage({action: null, change, failure: {kind: 'Syntax error', description: message}, owner,
-					ownerEmail: ownerEmailJustForSlackAtTheMoment, time: blockageOccurred})
+					ownerEmail, triager, triagerEmail, time: blockageOccurred})
 			}
 	
 			// stop the branch
@@ -2184,7 +2186,9 @@ export class NodeBot extends PerforceStatefulBot implements NodeBotInterface {
 		const blockageOccurred = existingBlockage ? existingBlockage.time : new Date
 
 		const ownerEmail = this.p4.getEmail(owner)
-		const blockage: Blockage = {change: pending.change, action: pending.action, failure, owner, ownerEmail, time: blockageOccurred, approval}
+		const triager = this.getBranchProperty({targetBranchName: pending.action.branch.upperName, property: "triager"})
+		const triagerEmail = triager && this.p4.getEmail(triager)
+		const blockage: Blockage = {change: pending.change, action: pending.action, failure, owner, ownerEmail, triager, triagerEmail, time: blockageOccurred, approval}
 
 		if (existingBlockage) {
 			this.nodeBotLogger.info(shortMessage + ' (already notified)')
@@ -2229,36 +2233,43 @@ export class NodeBot extends PerforceStatefulBot implements NodeBotInterface {
 			this.slackMessages.postNonConflictMessage(opts)
 		}
 	}
+	private getBranchProperty(property: string): any
+	private getBranchProperty(args: {targetBranchName?: string, property: string}): any
+	private getBranchProperty(property_or_args: string | {targetBranchName?: string, property: string}) {
+		const isArgs = typeof(property_or_args) !== 'string'
+		const {targetBranchName, property} = (isArgs ? property_or_args : {targetBranchName: undefined, property: property_or_args})
+		
+		const configKey = property as keyof BotConfig
+		const edgeKey   = property as keyof EdgeOptions
+		const nodeKey   = property as keyof NodeOptions
+
+		let prop: any | null = this.branchGraph.config[configKey]
+		if (targetBranchName) {
+			const targetBranch = this.branchGraph.getBranch(targetBranchName)
+			if (targetBranch && targetBranch.config[nodeKey] !== undefined)
+			{
+				prop = targetBranch.config[nodeKey]
+			}
+
+			const edgeProperties = this.branch.edgeProperties.get(targetBranchName)
+			if (edgeProperties && edgeProperties[edgeKey] !== undefined) {
+				prop = edgeProperties[edgeKey]
+			}
+		}
+		else if (this.branch.config[nodeKey] !== undefined)
+		{
+			prop = this.branch.config[nodeKey]
+		}
+		return prop
+	}
 
 	private async sendNagNotifications() {
 		let naggedSomeone = false
 		const now = Date.now()
 		for (const conflict of this.conflicts.getConflicts()) {
 
-			const getNagProperty = (property:string) => {
-						
-				const configKey = property as keyof BotConfig
-				const edgeKey   = property as keyof EdgeOptions
-				const nodeKey   = property as keyof NodeOptions
-
-				let prop: any | null = this.branchGraph.config[configKey]
-				if (conflict.targetBranchName) {
-					const targetBranch = this.branchGraph.getBranch(conflict.targetBranchName)
-					if (targetBranch && targetBranch.config[nodeKey] !== undefined)
-					{
-						prop = targetBranch.config[nodeKey]
-					}
-
-					const edgeProperties = this.branch.edgeProperties.get(conflict.targetBranchName)
-					if (edgeProperties && edgeProperties[edgeKey] !== undefined) {
-						prop = edgeProperties[edgeKey]
-					}
-				}
-				else if (this.branch.config[nodeKey] !== undefined)
-				{
-					prop = this.branch.config[nodeKey]
-				}
-				return prop
+			const getNagProperty = (property: string) => {
+				return this.getBranchProperty({targetBranchName: conflict.targetBranchName, property})
 			}
 			
 			const nagWhenBlocked = getNagProperty("nagWhenBlocked")
diff --git a/Engine/Extras/RoboMerge/v3/src/robo/notifications.ts b/Engine/Extras/RoboMerge/v3/src/robo/notifications.ts
--- a/Engine/Extras/RoboMerge/v3/src/robo/notifications.ts
+++ b/Engine/Extras/RoboMerge/v3/src/robo/notifications.ts
@@ -477,6 +477,29 @@ export class BotNotifications implements BotEventHandler {
 		return channelsToPostTo
 	}
 
+	async getTriager(triager: string, triagerEmail?: Promise<string|null>) {
+		if (triager && !isUserSlackGroup(triager))
+		{
+			const emailAddress = await triagerEmail
+			if (emailAddress)
+			{
+				const user = await this.slackMessages!.getSlackUser(emailAddress)
+				if (user) {
+					triager = `<@${user}>`
+				}
+				else {
+					triager = `@${triager}`
+					this.botNotificationsLogger.error(`Unable to look up triager from ${triager} (${emailAddress})`)
+				}
+			}
+			else {
+				triager = `@${triager}`
+				this.botNotificationsLogger.error(`Unable to look up email address for triager ${triager}`)
+			}
+		}
+		return triager
+	}
+
 	/** Conflict */
 	// considered if no Slack set up, should continue to add people to notify, but too complicated:
 	//		let's go all in on Slack. Fine for this to be async (but fire and forget)
@@ -508,10 +531,22 @@ export class BotNotifications implements BotEventHandler {
 		const channelPing = await getPingTarget(userEmail, blockage.owner, this.slackMessages)
 
 		const isBotUser = isUserAKnownBot(blockage.owner)
-		const text =
-			blockage.approval ?				`${channelPing}'s change needs to be approved\n\n${blockage.approval.settings.description}` :
-			isBotUser ? 									`Blockage caused by \`${blockage.owner}\` commit!` :
-															`${channelPing}, please resolve the following ${issue}:`
+		let text = ""
+		if (blockage.approval) {
+			text = `${channelPing}'s change needs to be approved\n\n${blockage.approval.settings.description}`
+		}
+		else if (isBotUser) {
+			if (blockage.triager) {
+				const triagerPing = await this.getTriager(blockage.triager, blockage.triagerEmail)
+				if (triagerPing) {
+					text = `${triagerPing} - `
+				}
+			}
+			text += `Blockage caused by \`${blockage.owner}\` commit!`
+		}
+		else {
+			text = `${channelPing}, please resolve the following ${issue}:`
+		}														
 
 		let message = this.makeSlackChannelMessage(
 			`${sourceBranch.name} blocked! (${issue})`,
@@ -679,27 +714,7 @@ export class BotNotifications implements BotEventHandler {
 				}
 			}
 
-			let triager = evt.triager
-			if (triager && !isUserSlackGroup(triager))
-			{
-				const emailAddress = await evt.triagerEmail
-				if (emailAddress)
-				{
-					const user = await this.slackMessages.getSlackUser(emailAddress)
-					if (user) {
-						triager = `<@${user}>`
-					}
-					else {
-						triager = `@${triager}`
-						this.botNotificationsLogger.error(`Unable to look up triager from ${triager} (${emailAddress})`)
-					}
-				}
-				else {
-					triager = `@${triager}`
-					this.botNotificationsLogger.error(`Unable to look up email address for triager ${triager}`)
-				}
-			}
-
+			const triager = await this.getTriager(evt.triager, evt.triagerEmail)
 			const channelPing = await getPingTarget(await evt.ownerEmail, conflict.owner, this.slackMessages)
 
 			const nagMessage: SlackMessage = {
