Fix messages loading incorrectly from DB

diff --git a/Engine/Extras/RoboMerge/v3/src/robo/persistedmessages.ts b/Engine/Extras/RoboMerge/v3/src/robo/persistedmessages.ts
--- a/Engine/Extras/RoboMerge/v3/src/robo/persistedmessages.ts
+++ b/Engine/Extras/RoboMerge/v3/src/robo/persistedmessages.ts
@@ -150,9 +150,9 @@ class DBPersistedSlackMessages extends PersistedSlackMessages {
 		const messageRecord = await SlackMessagesModel.find(index)
 		if (messageRecord.length == 1) {
 			const persistedMessage = new PersistedSlackMessage({
-						messageOpts: messageRecord.message, 
-						permalink: messageRecord.permalink,
-						timestamp: messageRecord.timestamp
+						messageOpts: messageRecord[0].message, 
+						permalink: messageRecord[0].permalink,
+						timestamp: messageRecord[0].timestamp
 					}, key, this)
 			this.messageCache[cacheKey] = persistedMessage
 			return persistedMessage
