Fix syntax error nagging failing to identify channel to write to

diff --git a/Engine/Extras/RoboMerge/v3/src/robo/branch-interfaces.ts b/Engine/Extras/RoboMerge/v3/src/robo/branch-interfaces.ts
--- a/Engine/Extras/RoboMerge/v3/src/robo/branch-interfaces.ts
+++ b/Engine/Extras/RoboMerge/v3/src/robo/branch-interfaces.ts
@@ -357,6 +357,7 @@ export type EndIntegratingToGateEvent = {
 }
 
 export type NagBlockageEvent = {
+	sourceBranch: Branch
 	conflict: PersistentConflict
 	triager: string
 	triagerEmail: Promise<string>
diff --git a/Engine/Extras/RoboMerge/v3/src/robo/nodebot.ts b/Engine/Extras/RoboMerge/v3/src/robo/nodebot.ts
--- a/Engine/Extras/RoboMerge/v3/src/robo/nodebot.ts
+++ b/Engine/Extras/RoboMerge/v3/src/robo/nodebot.ts
@@ -2293,7 +2293,8 @@ export class NodeBot extends PerforceStatefulBot implements NodeBotInterface {
 						naggedSomeone = true
 						const triager = getNagProperty("triager")
 						this.eventTriggers.nagBlockage({
-							conflict, 
+							sourceBranch: this.branch,
+							conflict,
 							triager,
 							triagerEmail: triager && this.p4.getEmail(triager),
 							ownerEmail: this.p4.getEmail(conflict.owner)
diff --git a/Engine/Extras/RoboMerge/v3/src/robo/notifications.ts b/Engine/Extras/RoboMerge/v3/src/robo/notifications.ts
--- a/Engine/Extras/RoboMerge/v3/src/robo/notifications.ts
+++ b/Engine/Extras/RoboMerge/v3/src/robo/notifications.ts
@@ -724,7 +724,7 @@ export class BotNotifications implements BotEventHandler {
 				mrkdwn: true
 			}
 
-			const channelsToPostTo = this.getChannelsToPostTo(conflict.blockedBranchName, conflict.targetBranchName)
+			const channelsToPostTo = this.getChannelsToPostTo(evt.sourceBranch, conflict.targetBranchName)
 
 			this.botNotificationsLogger.info(`Sending nag notification to ${triager ? triager + ' and ' : ''}${conflict.owner} after ${Math.floor(ageMinutes)} minutes`)
 			const branchName = (conflict.kind === "Syntax error" ? conflict.kind : conflict.targetBranchName || conflict.blockedBranchName)
