Clean up nag pinging

diff --git a/Engine/Extras/RoboMerge/v3/src/robo/notifications.ts b/Engine/Extras/RoboMerge/v3/src/robo/notifications.ts
--- a/Engine/Extras/RoboMerge/v3/src/robo/notifications.ts
+++ b/Engine/Extras/RoboMerge/v3/src/robo/notifications.ts
@@ -90,7 +90,9 @@ export async function getPingTarget(userEmail: string|null, owner: string, slack
 	? `<@${slackUser}>` 
 	: isUserSlackGroup(owner)
 		? owner 
-		: `@${owner}`
+		: isUserAKnownBot(owner)
+			? ""
+			: `@${owner}`
 }
 
 function formatDuration(durationSeconds: number) {
@@ -717,8 +719,20 @@ export class BotNotifications implements BotEventHandler {
 			const triager = await this.getTriager(evt.triager, evt.triagerEmail)
 			const channelPing = await getPingTarget(await evt.ownerEmail, conflict.owner, this.slackMessages)
 
+			let messageTargets = []
+			let loggerTargets = []
+			if (triager && triager != channelPing) {
+				messageTargets.push(triager)
+				loggerTargets.push(evt.triager)
+			}
+			if (channelPing != "") {
+				messageTargets.push(channelPing)
+				loggerTargets.push(conflict.owner)
+			}
+			messageTargets.push("Robomerge")
+
 			const nagMessage: SlackMessage = {
-				text: `${triager ? triager + ' ' : ''}${channelPing} RoboMerge blocked for more than ${timeDesc}`,
+				text: `${messageTargets.join(' ')} blocked for more than ${timeDesc}`,
 				style: SlackMessageStyles.DANGER,
 				channel: this.slackChannel,
 				mrkdwn: true
@@ -726,7 +740,7 @@ export class BotNotifications implements BotEventHandler {
 
 			const channelsToPostTo = this.getChannelsToPostTo(evt.sourceBranch, conflict.targetBranchName)
 
-			this.botNotificationsLogger.info(`Sending nag notification to ${triager ? triager + ' and ' : ''}${conflict.owner} after ${Math.floor(ageMinutes)} minutes`)
+			this.botNotificationsLogger.info(`Sending nag notification${loggerTargets.length > 0 ? ' to ' + loggerTargets.join(' and ') : ''} after ${Math.floor(ageMinutes)} minutes`)
 			const branchName = (conflict.kind === "Syntax error" ? conflict.kind : conflict.targetBranchName || conflict.blockedBranchName)
 			for (let slackChannel of channelsToPostTo) {
 				nagMessage.channel = slackChannel
