Robomerge: Changes for migration to the cloud

* new CLI/Envvars
* behindLB - robomerge sits behind an LB so it is secure but doesn't need to start as https & use certs
* webserverPort - can override the default 8877/4433 port depending on secure or not, defaults to -1 ie: existing behaviour
* expose noTLS to the watchdog as it should respect that flag too
* watchdog in 'non secure' mode should use the watchdogPort too
* webserver uses behindLb to start up with the right security level if protoal is https we are fully secure

diff --git a/Engine/Extras/RoboMerge/v3/src/common/webserver.ts b/Engine/Extras/RoboMerge/v3/src/common/webserver.ts
--- a/Engine/Extras/RoboMerge/v3/src/common/webserver.ts
+++ b/Engine/Extras/RoboMerge/v3/src/common/webserver.ts
@@ -113,11 +113,13 @@ export class WebServer {
 	protected serverLogger: ContextualLogger
 
 	public secure = false
+	protected behindLoadBalancer = false
 
 	temp_debug_id: number
 
-	constructor(parentLogger: ContextualLogger) {
+	constructor(parentLogger: ContextualLogger, behindLoadBalancer?: boolean) {
 		this.serverLogger = parentLogger.createChild('WebServer')
+		this.behindLoadBalancer = behindLoadBalancer || false
 	}
 
 	addApp(appType: AppConstructor, inMakeApp?: MakeApp) {
@@ -149,11 +151,12 @@ export class WebServer {
 		}
 
 		protocol = protocol || 'https'
-		this.secure = protocol.toLowerCase() !== 'http'
+		const isSelfSecure = protocol.toLowerCase() !== 'http'
+		this.secure = this.behindLoadBalancer || isSelfSecure
 		const port = optPort || (this.secure ? 443 : 80)
 
 
-		if (this.secure) {
+		if (this.secure && isSelfSecure) {
 			// cast because @types/node doesn't seem to know about this constructor
 			this.server = new (https.Server as any)(certFiles!)
 		}
diff --git a/Engine/Extras/RoboMerge/v3/src/robo/robo.ts b/Engine/Extras/RoboMerge/v3/src/robo/robo.ts
--- a/Engine/Extras/RoboMerge/v3/src/robo/robo.ts
+++ b/Engine/Extras/RoboMerge/v3/src/robo/robo.ts
@@ -119,6 +119,13 @@ const COMMAND_LINE_ARGS: {[param: string]: Arg<any>} = {
 		dflt: false
 	},
 
+	behindLoadBalancer: {
+		match: /^(-behindLoadBalancer)$/,
+		parse: _str => true,
+		env: 'ROBO_BEHIND_LOAD_BALANCER',
+		dflt: false
+	},
+
 	useSlackInDev: {
 		match: /^(-useSlackInDev)$/,
 		parse: _str => true,
@@ -239,6 +246,15 @@ const COMMAND_LINE_ARGS: {[param: string]: Arg<any>} = {
 			return parseInt(str)
 		},
 		dflt: 60
+	},
+	webserverPort: {
+		match: /^-webServerPort=(\d+)$/,
+		env: 'ROBO_WEB_SERVER_PORT',
+		parse: (str: string) => {
+			const val = parseInt(str)
+			return isNaN(val) ? -1 : val
+		},
+		dflt: -1
 	}
 };
 
@@ -535,8 +551,12 @@ else {
 	const ws = new RoboServer(args.externalUrl, sendMessage, () => roboStartupLogger.getLogTail(),
 		() => roboStartupLogger.info('Received: getLastCrash'),
 		() => roboStartupLogger.info('Received: stopBot'),
-		() => roboStartupLogger.info('Received: startBot'));
+		() => roboStartupLogger.info('Received: startBot'),
+		args.behindLoadBalancer);
 
+	if (args.noTLS) {
+		roboStartupLogger.info(`Server will start in non-secure mode...`)
+	}
 
 	// TODO: Please make this better
 	const certFiles = {
@@ -545,7 +565,7 @@ else {
 	}
 
 	const protocol = args.noTLS ? 'http' : 'https'
-	const port = args.noTLS ? 8877 : 4433
+	const port = args.webserverPort != -1 ? args.webserverPort : args.noTLS ? 8877 : 4433
 	ws.open(port, protocol, certFiles as CertFiles).then(() =>
 		roboStartupLogger.info(`Running in-process web server (${protocol}) on port ${port}`)
 	);
diff --git a/Engine/Extras/RoboMerge/v3/src/robo/roboserver.ts b/Engine/Extras/RoboMerge/v3/src/robo/roboserver.ts
--- a/Engine/Extras/RoboMerge/v3/src/robo/roboserver.ts
+++ b/Engine/Extras/RoboMerge/v3/src/robo/roboserver.ts
@@ -63,6 +63,12 @@ let ENVIRONMENT: {[param: string]: any}
 			parse: str => str === "false" ? false : true,
 			env: 'LDAP_SIGN_IN',
 			dflt: true
+		},
+		behindLoadBalancer: {
+			match: /^(-behindLoadBalancer)$/,
+			parse: _str => true,
+			env: 'ROBO_BEHIND_LOAD_BALANCER',
+			dflt: false
 		}
 	}
 
@@ -123,8 +129,8 @@ interface BackendFunctions {
 export class RoboServer {
 
 	constructor(externalUrl: string, sendMessage: (msg: string, args?: any[]) => any, getLogTail: Function | null, 
-		getLastCrash: Function, stopBot: Function, startBot: Function) {
-		this.server = new WebServer(this.roboserverLogger)
+		getLastCrash: Function, stopBot: Function, startBot: Function, behindLoadBalancer?: boolean) {
+		this.server = new WebServer(this.roboserverLogger, behindLoadBalancer)
 		this.functions = {
 			sendMessage: sendMessage,
 			getLogTail: getLogTail,
@@ -624,7 +630,7 @@ class RoboWebApp implements AppInterface {
 
 	// Simple test message which assumes user has an email address taking the form of [REDACTED_EMAIL]
 	// Example usage:
-	// curl -X POST -k -v --cookie "auth=AUTH_COOKIE_FROM_BROWSER" https://robomerge.companyname.net/api/test/directmessage/jake.romigh
+	// curl -X POST -k -v --cookie "auth=AUTH_COOKIE_FROM_BROWSER" https://robomerge.companyname.net/api/test/directmessage/user
 	@SecureHandler('POST', '/api/test/directmessage/*', {requiredTags: ['admin']})
 	sendTestDirectMessage(user : string) {
 		this.webAppLogger.verbose('Received request for direct message')
diff --git a/Engine/Extras/RoboMerge/v3/src/robo/watchdog.ts b/Engine/Extras/RoboMerge/v3/src/robo/watchdog.ts
--- a/Engine/Extras/RoboMerge/v3/src/robo/watchdog.ts
+++ b/Engine/Extras/RoboMerge/v3/src/robo/watchdog.ts
@@ -62,6 +62,12 @@ const COMMAND_LINE_ARGS: {[param: string]: (Arg<any>)} = {
 		},
 		dflt: 'unknown'
 	},
+	noTLS: {
+		match: /^(-noTLS)$/,
+		parse: _str => true,
+		env: 'ROBO_NO_TLS',
+		dflt: false
+	},
 	watchdogPort: {
 		match: /^-watchdogPort=(\d+)$/,
 		env: 'ROBO_WATCHDOG_PORT',
@@ -75,6 +81,12 @@ const COMMAND_LINE_ARGS: {[param: string]: (Arg<any>)} = {
 		match: /^-sentryDsn=(.+)$/,
 		env: 'SENTRY_DSN',
 		dflt: 'https:[REDACTED_EMAIL]/1432517' // Robomerge sentry project: https://sentry.io/organizations/to/issues/?project=1432517
+	},
+	behindLoadBalancer: {
+		match: /^(-behindLoadBalancer)$/,
+		parse: _str => true,
+		env: 'ROBO_BEHIND_LOADBALANCER',
+		dflt: false
 	}
 }
 
@@ -119,7 +131,7 @@ class Watchdog {
 		// expose a webserver that can be used to check status
 		this.statusServer = new RoboServer(args.externalUrl,
 			(name: string, args?: any[]) => this.sendMessage(name, args), 
-			() => this.watchdogLogger.getLogTail(), () => this.getLastCrash(), () => this.stopBot(), () => this.startBot()
+			() => this.watchdogLogger.getLogTail(), () => this.getLastCrash(), () => this.stopBot(), () => this.startBot(), args.behindLoadBalancer
 		)
 
 		setImmediate(() => this.spawnEntryPoint())
@@ -198,14 +210,23 @@ class Watchdog {
 	startServer() {
 		// for now, require HTTPS vault settings - eventually have a proper dev
 		// setting that can allow local testing without
-		try {
-			this.watchdogLogger.info('Checking for files in ' + Session.VAULT_PATH)
-			fs.readdirSync(Session.VAULT_PATH)
-		}
-		catch (err) {
-			console.log(err)
-			this.statusServer.open(8877, 'http').then(() => 
-				this.watchdogLogger.warn(`HTTP web server opened on port 8877`)
+		if (!args.noTLS) {
+			try {
+				this.watchdogLogger.info(`Checking for files in ${Session.VAULT_PATH}`)
+				fs.readdirSync(Session.VAULT_PATH)
+			}
+			catch (err) {
+				console.log(err)
+				this.statusServer.open(8877, 'http').then(() => 
+					this.watchdogLogger.warn(`HTTP web server opened on port 8877`)
+				)
+				return
+			}
+		} 
+		else {
+			// We are running with no certs, presumabily behind a load balancer, just open as http and carry on..
+			this.statusServer.open(args.watchdogPort, 'http').then(() => 
+				this.watchdogLogger.warn(`HTTP web server opened on port ${args.watchdogPort}`)
 			)
 			return
 		}
