DB related logging

diff --git a/Engine/Extras/RoboMerge/v3/src/robo/robo.ts b/Engine/Extras/RoboMerge/v3/src/robo/robo.ts
--- a/Engine/Extras/RoboMerge/v3/src/robo/robo.ts
+++ b/Engine/Extras/RoboMerge/v3/src/robo/robo.ts
@@ -275,9 +275,12 @@ settingsInit(args)
 if (RoboArgs.useMongo()) {
 	var mongoose = require('mongoose')
 	mongoose.connect(args.mongoDB_URI)
-	mongoose.connection.on('connected', () => doDBPersistedMessagesHousekeeping())
+	mongoose.connection.on('connected', () => {
+		roboStartupLogger.info(`MongoDB Connected`)
+		doDBPersistedMessagesHousekeeping()
+	})
 	mongoose.connection.on('error', (err: any) => {
-		console.log(`MongoDB Connection Error: ${err}`)
+		roboStartupLogger.error(`MongoDB Connection Error: ${err}`)
 		process.exit(1);
 	})
 }
diff --git a/Engine/Extras/RoboMerge/v3/src/robo/settings.ts b/Engine/Extras/RoboMerge/v3/src/robo/settings.ts
--- a/Engine/Extras/RoboMerge/v3/src/robo/settings.ts
+++ b/Engine/Extras/RoboMerge/v3/src/robo/settings.ts
@@ -141,6 +141,7 @@ export class Settings {
 		if (RoboArgs.useMongo()) {
 			// first try to get the bot data from the DB
 			const botRecords = await BotModel.find({botname: botname})
+			settings.settingsLogger.info(`Loaded ${botRecords.length} node records`)
 			if (botRecords.length > 1) {
 				// starting with 3.3.0 the individual node records themselves are being persisted/versioned, 
 				// not the top level object itself, so we can set the object to that version to indicate there
@@ -154,6 +155,7 @@ export class Settings {
 			else if (botRecords.length == 1) {
 				const objVersion = semver.coerce(String(botRecords[0].details.version))!
 				if (semver.lt(objVersion, '3.3.0')) {
+					settings.settingsLogger.info(`Migrated from per bot to per node records`)
 					settings.object = botRecords[0].details
 					BotModel.deleteOne({_id: botRecords[0]._id})
 				}
@@ -215,6 +217,7 @@ export class Settings {
 
 		// Transition slack messages out of the settings object and into the database
 		if (loadedFromFile && RoboArgs.useMongo()) {
+			settings.settingsLogger.info("Migrated settings from json to MongoDB")
 			await migrateMessagesToDB(settings)
 		}
 
