Restoring some missing changes around nonTLS for startup

diff --git a/Engine/Extras/RoboMerge/v3/src/robo/robo.ts b/Engine/Extras/RoboMerge/v3/src/robo/robo.ts
--- a/Engine/Extras/RoboMerge/v3/src/robo/robo.ts
+++ b/Engine/Extras/RoboMerge/v3/src/robo/robo.ts
@@ -558,7 +558,9 @@ else {
 		args.behindLoadBalancer);
 
 	if (args.noTLS) {
-		roboStartupLogger.info(`Server will start in non-secure mode...`)
+		roboStartupLogger.info(`Server will start in self-non-secure mode...`)
+	} else {
+		roboStartupLogger.info(`Server will start in self-secure mode...`)
 	}
 
 	// TODO: Please make this better
diff --git a/Engine/Extras/RoboMerge/v3/src/robo/watchdog.ts b/Engine/Extras/RoboMerge/v3/src/robo/watchdog.ts
--- a/Engine/Extras/RoboMerge/v3/src/robo/watchdog.ts
+++ b/Engine/Extras/RoboMerge/v3/src/robo/watchdog.ts
@@ -210,21 +210,22 @@ class Watchdog {
 	startServer() {
 		// for now, require HTTPS vault settings - eventually have a proper dev
 		// setting that can allow local testing without
+		
+		const protocol = args.noTLS ? 'http' : 'https'
 		if (!args.noTLS) {
 			try {
-				this.watchdogLogger.info(`Checking for files in ${Session.VAULT_PATH}`)
+				this.watchdogLogger.info('Checking for files in ' + Session.VAULT_PATH)
 				fs.readdirSync(Session.VAULT_PATH)
 			}
 			catch (err) {
 				console.log(err)
-				this.statusServer.open(8877, 'http').then(() => 
-					this.watchdogLogger.warn(`HTTP web server opened on port 8877`)
+				this.statusServer.open(args.watchdogPort, 'http').then(() => 
+					this.watchdogLogger.warn(`HTTP web server opened on port ${args.watchdogPort}`)
 				)
 				return
 			}
 		} 
 		else {
-			// We are running with no certs, presumabily behind a load balancer, just open as http and carry on..
 			this.statusServer.open(args.watchdogPort, 'http').then(() => 
 				this.watchdogLogger.warn(`HTTP web server opened on port ${args.watchdogPort}`)
 			)
@@ -232,12 +233,12 @@ class Watchdog {
 		}
 
 		const certFiles = {
-			key: fs.readFileSync(`${Session.VAULT_PATH}/${tlsKeyFilename}`, 'ascii'),
-			cert: fs.readFileSync('./certs/robomerge.pem', 'ascii')
+			key: args.noTLS ? "" : fs.readFileSync(`${Session.VAULT_PATH}/${tlsKeyFilename}`, 'ascii'),
+			cert: args.noTLS ? "" : fs.readFileSync('./certs/robomerge.pem', 'ascii')
 		}
 		// Removed ca: [fs.readFileSync('./certs/thwate.cer')] - JR 10/4
 
-		this.statusServer.open(args.watchdogPort, 'https', certFiles).then(() => 
+		this.statusServer.open(args.watchdogPort, protocol, certFiles).then(() => 
 			this.watchdogLogger.info(`Web server opened on port ${args.watchdogPort}`)
 		)
 	}
