Robomerge: fixing loadbalancer args and improving logging

diff --git a/Engine/Extras/RoboMerge/v3/src/common/webserver.ts b/Engine/Extras/RoboMerge/v3/src/common/webserver.ts
--- a/Engine/Extras/RoboMerge/v3/src/common/webserver.ts
+++ b/Engine/Extras/RoboMerge/v3/src/common/webserver.ts
@@ -168,6 +168,14 @@ export class WebServer {
 
 		this.serverLogger.info(`Web server listening on port ${port} over ${protocol}`)
 
+		if (this.behindLoadBalancer) {
+			this.serverLogger.info(`Web server secure behind loadbalancer, running as http but presuming secure.`)
+		} else if (isSelfSecure) {
+			this.serverLogger.info(`Web server using local installed certs for https.`)
+		} else {
+			this.serverLogger.warn(`Web server may be unsecure`)
+		}
+
 		return new Promise<void>((done, _fail) => this.server!.listen(port, done))
 	}
 
diff --git a/Engine/Extras/RoboMerge/v3/src/robo/robo.ts b/Engine/Extras/RoboMerge/v3/src/robo/robo.ts
--- a/Engine/Extras/RoboMerge/v3/src/robo/robo.ts
+++ b/Engine/Extras/RoboMerge/v3/src/robo/robo.ts
@@ -121,8 +121,8 @@ const COMMAND_LINE_ARGS: {[param: string]: Arg<any>} = {
 
 	behindLoadBalancer: {
 		match: /^(-behindLoadBalancer)$/,
-		parse: _str => true,
-		env: 'ROBO_BEHIND_LOAD_BALANCER',
+		parse: str => str === "false" ? false : true,
+		env: 'ROBO_BEHIND_LOADBALANCER',
 		dflt: false
 	},
 
diff --git a/Engine/Extras/RoboMerge/v3/src/robo/roboserver.ts b/Engine/Extras/RoboMerge/v3/src/robo/roboserver.ts
--- a/Engine/Extras/RoboMerge/v3/src/robo/roboserver.ts
+++ b/Engine/Extras/RoboMerge/v3/src/robo/roboserver.ts
@@ -66,8 +66,8 @@ let ENVIRONMENT: {[param: string]: any}
 		},
 		behindLoadBalancer: {
 			match: /^(-behindLoadBalancer)$/,
-			parse: _str => true,
-			env: 'ROBO_BEHIND_LOAD_BALANCER',
+			parse: str => str === "false" ? false : true,
+			env: 'ROBO_BEHIND_LOADBALANCER',
 			dflt: false
 		}
 	}
diff --git a/Engine/Extras/RoboMerge/v3/src/robo/watchdog.ts b/Engine/Extras/RoboMerge/v3/src/robo/watchdog.ts
--- a/Engine/Extras/RoboMerge/v3/src/robo/watchdog.ts
+++ b/Engine/Extras/RoboMerge/v3/src/robo/watchdog.ts
@@ -84,7 +84,7 @@ const COMMAND_LINE_ARGS: {[param: string]: (Arg<any>)} = {
 	},
 	behindLoadBalancer: {
 		match: /^(-behindLoadBalancer)$/,
-		parse: _str => true,
+		parse: str => str === 'false' ? false : true,
 		env: 'ROBO_BEHIND_LOADBALANCER',
 		dflt: false
 	}
