Restore apparently necessary intermediate function but keep out the logging

diff --git a/Engine/Extras/RoboMerge/v3/src/robo/nodebot.ts b/Engine/Extras/RoboMerge/v3/src/robo/nodebot.ts
--- a/Engine/Extras/RoboMerge/v3/src/robo/nodebot.ts
+++ b/Engine/Extras/RoboMerge/v3/src/robo/nodebot.ts
@@ -412,7 +412,7 @@ export class NodeBot extends PerforceStatefulBot implements NodeBotInterface {
 						mrkdwn: true
 					}
 					
-					let emailAddress = await this.p4.getEmail(fromQueue.who)
+					let emailAddress = await this.findEmail(fromQueue.who)
 					this.slackMessages.postDM(emailAddress, fromQueue.cl, this.branch, dm)
 				}
 				else {
@@ -436,7 +436,7 @@ export class NodeBot extends PerforceStatefulBot implements NodeBotInterface {
 						mrkdwn: true
 					}
 					
-					let emailAddress = await this.p4.getEmail(fromQueue.who)
+					let emailAddress = await this.findEmail(fromQueue.who)
 					this.slackMessages.postDM(emailAddress, fromQueue.cl, this.branch, dm)
 				}
 				else {
@@ -475,7 +475,7 @@ export class NodeBot extends PerforceStatefulBot implements NodeBotInterface {
 					mrkdwn: true
 				}
 				
-				let emailAddress = await this.p4.getEmail(fromQueue.who)
+				let emailAddress = await this.findEmail(fromQueue.who)
 				this.slackMessages.postDM(emailAddress, fromQueue.cl, this.branch, dm)
 			}
 
@@ -1060,7 +1060,7 @@ export class NodeBot extends PerforceStatefulBot implements NodeBotInterface {
 			filesByUser.forEach(async (files, author) => {
 
 				let userToNotify = `@${author}`
-				const emailAddress = await this.p4.getEmail(author)
+				const emailAddress = await this.findEmail(author)
 				if (emailAddress)
 				{
 					const user = await this.slackMessages!.getSlackUser(emailAddress)
@@ -1381,7 +1381,7 @@ export class NodeBot extends PerforceStatefulBot implements NodeBotInterface {
 				mrkdwn: true
 			}
 			
-			this.slackMessages.postDM(this.p4.getEmail(event.change.owner!), event.change.cl, this.branch, dm)			
+			this.slackMessages.postDM(this.findEmail(event.change.owner!), event.change.cl, this.branch, dm)			
 		}
 	}
 
@@ -1395,6 +1395,10 @@ export class NodeBot extends PerforceStatefulBot implements NodeBotInterface {
 		}
 	}
 
+	findEmail = (user: string) => {
+		return this.p4.getEmail(user);
+	}
+
 	// Listener method to email a blockage owner when the branch encounters said blockage,
 	// depending on NodeBot configuration
 	private emailBlockageCulprit(blockage: Blockage) {
@@ -1447,7 +1451,7 @@ export class NodeBot extends PerforceStatefulBot implements NodeBotInterface {
 				mrkdwn: true
 			}
 			
-			let emailAddress = await this.p4.getEmail(owner)
+			let emailAddress = await this.findEmail(owner)
 			this.slackMessages.postDM(emailAddress, pendingChange.newCl, this.branch, dm)
 		}
 		else
@@ -1725,7 +1729,7 @@ export class NodeBot extends PerforceStatefulBot implements NodeBotInterface {
 	private _sendGenericEmail(recipients: Recipients, subject: string, intro: string, message: string) {
 		const params : MailParams = { intro, message, botname: this.graphBotName, 
 			robomergeRootUrl: { value: this.externalRobomergeUrl, noEscape: true } }
-		this.mailer.sendEmail('generic-email-template', recipients, subject, params, this.p4.getEmail)
+		this.mailer.sendEmail('generic-email-template', recipients, subject, params, this.findEmail)
 	}
 	private _sendBlockageEmail(blockage: Blockage) {
 		const recipients = new Recipients(blockage.owner)
@@ -1769,7 +1773,7 @@ export class NodeBot extends PerforceStatefulBot implements NodeBotInterface {
 		// Add list to email parameters
 		params.resolutionOptions = { value: resolutionOptionsStr, noEscape: true }
 
-		this.mailer.sendEmail('new-blockage', recipients, subject, params, this.p4.getEmail)
+		this.mailer.sendEmail('new-blockage', recipients, subject, params, this.findEmail)
 	}
 
 	sendErrorEmail(recipients: Recipients, subject: string, message: string, errorPreface='RoboMerge needs your help:'): void {
@@ -2097,7 +2101,7 @@ export class NodeBot extends PerforceStatefulBot implements NodeBotInterface {
 					mrkdwn: true
 				}
 				
-				let emailAddress = await this.p4.getEmail(owner)
+				let emailAddress = await this.findEmail(owner)
 				this.slackMessages.postDM(emailAddress, change.cl, this.branch, dm)				
 			}
 			else {
@@ -2157,7 +2161,7 @@ export class NodeBot extends PerforceStatefulBot implements NodeBotInterface {
 					mrkdwn: true
 				}
 				
-				let emailAddress = await this.p4.getEmail(owner)
+				let emailAddress = await this.findEmail(owner)
 				await this.slackMessages.postDM(emailAddress, pending.newCl, this.branch, dm)
 
 				if (failure.details) {
@@ -2296,8 +2300,8 @@ export class NodeBot extends PerforceStatefulBot implements NodeBotInterface {
 							sourceBranch: this.branch,
 							conflict,
 							triager,
-							triagerEmail: triager && this.p4.getEmail(triager),
-							ownerEmail: this.p4.getEmail(conflict.owner)
+							triagerEmail: triager && this.findEmail(triager),
+							ownerEmail: this.findEmail(conflict.owner)
 						})
 					}
 				}
