Add environment variables to be able to run functional tests locally with loadbalancer/port changes

diff --git a/Engine/Extras/RoboMerge/v3/run-tests.sh b/Engine/Extras/RoboMerge/v3/run-tests.sh
--- a/Engine/Extras/RoboMerge/v3/run-tests.sh
+++ b/Engine/Extras/RoboMerge/v3/run-tests.sh
@@ -89,7 +89,10 @@ docker run -d -p 8877:8877 -p 8811:8811 -p 25:25 -h robomerge_functtest --name r
     -e ROBO_USE_SLACK_IN_DEV=true \
     -e ROBO_NO_MAIL=true \
     -e ROBO_NO_TLS=true  \
+    -e ROBO_RUNNING_FUNCTIONAL_TESTS=true                 \
     -e ROBO_EXTERNAL_URL=http://robomerge_functtest:8877  \
+    -e ROBO_WEB_SERVER_PORT=8877                          \
+    -e ROBO_WATCHDOG_PORT=8877                            \
     -e BOTNAME=ft1,ft2,ft3,ft4,targets  \
     -e ROBO_LOG_LEVEL=info \
     -e ROBO_SLACK_DOMAIN=http://localhost:8811 \
