When doing a null merge branches and deletes are integrated requiring resolves so as to be able to use -ay

diff --git a/Engine/Extras/RoboMerge/v3/functional_tests/framework/src/tests/test-flags.ts b/Engine/Extras/RoboMerge/v3/functional_tests/framework/src/tests/test-flags.ts
--- a/Engine/Extras/RoboMerge/v3/functional_tests/framework/src/tests/test-flags.ts
+++ b/Engine/Extras/RoboMerge/v3/functional_tests/framework/src/tests/test-flags.ts
@@ -2,6 +2,7 @@
 import { P4Util } from '../framework'
 import { MultipleDevAndReleaseTestBase } from '../MultipleDevAndReleaseTestBase'
 
+const initialContent = 'Initial content'
 const contentAfterFirstEdit = 'Initial content\n\nFirst addition'
 
 export class TestFlags extends MultipleDevAndReleaseTestBase {
@@ -9,8 +10,12 @@ export class TestFlags extends MultipleDevAndReleaseTestBase {
 		await this.p4.depot('stream', this.depotSpec())
 		await this.createStreamsAndWorkspaces()
 
-		await P4Util.addFileAndSubmit(this.getClient('Main', 'testuser1'), 'test.txt', 'Initial content')
-
+		const client = this.getClient('Main', 'testuser1')
+		await Promise.all([
+			P4Util.addFile(client, 'test.txt', initialContent),
+			P4Util.addFile(client, 'test_delete.txt', initialContent)
+		])
+		await client.submit('Adding initial files')
 		await this.initialPopulate()
 	}
 
@@ -20,8 +25,11 @@ export class TestFlags extends MultipleDevAndReleaseTestBase {
 		const skipPootle = '-' + this.fullBranchName('Dev-Pootle')
 		return r1client.sync()
 		.then(() => P4Util.editFileAndSubmit(r1client, 'test.txt', contentAfterFirstEdit, skipPootle))
-		.then(() => P4Util.editFileAndSubmit(r1client, 'test.txt', 
-				'Initial content\n\nFirst addition\n\nSecond addition', nullMergePerkin + ' ' + skipPootle))
+		.then(() => P4Util.editFile(r1client, 'test.txt', 
+				'Initial content\n\nFirst addition\n\nSecond addition'))
+		.then(() => P4Util.addFile(r1client, 'test_add.txt', initialContent))
+		.then(() => r1client.delete('test_delete.txt'))
+		.then(() => P4Util.submit(r1client, '#robomerge ' + nullMergePerkin + ' ' + skipPootle))
 	}
 
 	async checkNumForwardedCommands(stream: string, expected: number) {
@@ -55,6 +63,26 @@ export class TestFlags extends MultipleDevAndReleaseTestBase {
 				if (contents !== contentAfterFirstEdit) {
 					throw new Error('null merge not so null: ' + contents)
 				}
+			}),
+
+			this.checkHeadRevision('Main', 'test_add.txt', 1),
+			this.checkHeadRevision('Release-2.0', 'test_add.txt', 1),
+			this.checkHeadRevision('Dev-Perkin', 'test_add.txt', 1),
+			this.checkHeadRevision('Dev-Pootle', 'test_add.txt', 0),
+			this.getClient('Dev-Perkin').print('test_add.txt').then(contents => {
+				if (contents !== '') {
+					throw new Error('null merge of add not so null: ' + contents)
+				}
+			}),
+
+			this.checkHeadRevision('Main', 'test_delete.txt', 2),
+			this.checkHeadRevision('Release-2.0', 'test_delete.txt', 2),
+			this.checkHeadRevision('Dev-Perkin', 'test_delete.txt', 2),
+			this.checkHeadRevision('Dev-Pootle', 'test_delete.txt', 1),
+			this.getClient('Dev-Perkin').print('test_delete.txt').then(contents => {
+				if (contents !== initialContent) {
+					throw new Error('null merge of delete not so null: ' + contents)
+				}
 			})
 		])
 	}
diff --git a/Engine/Extras/RoboMerge/v3/src/common/perforce.ts b/Engine/Extras/RoboMerge/v3/src/common/perforce.ts
--- a/Engine/Extras/RoboMerge/v3/src/common/perforce.ts
+++ b/Engine/Extras/RoboMerge/v3/src/common/perforce.ts
@@ -300,6 +300,7 @@ export interface EditChangeOpts {
 export interface IntegrateOpts {
 	edgeServerAddress?: string
 	virtual?: boolean
+	nullMerge?: boolean
 }
 
 export interface IntegratedOpts {
@@ -1292,6 +1293,12 @@ export class PerforceContext {
 			cmdList.push('-v')
 		}
 
+		// If we're doing a null merge then the branches and deletes need to be scheduled for resolve
+		if (opts.nullMerge) {
+			cmdList.push('-Rb')
+			cmdList.push('-Rd')
+		}
+
 		let noSuchFilesPossible = false // Helper variable for error catching
 		// Branchspec -- takes priority above other possibilities
 		if (source.branchspec) {
diff --git a/Engine/Extras/RoboMerge/v3/src/robo/edgebot.ts b/Engine/Extras/RoboMerge/v3/src/robo/edgebot.ts
--- a/Engine/Extras/RoboMerge/v3/src/robo/edgebot.ts
+++ b/Engine/Extras/RoboMerge/v3/src/robo/edgebot.ts
@@ -579,7 +579,7 @@ class EdgeBotImpl extends PerforceStatefulBot {
 
 		this.currentIntegrationStartTimestamp = Date.now()
 		await this.p4.sync(info.targetWorkspace, this.targetBranch.rootPath + '#0', {edgeServerAddress})
-		const [mode, results] = await this.p4.integrate(info.targetWorkspace, source, changenum, integTarget, {edgeServerAddress, virtual: doVirtualMerge})
+		const [mode, results] = await this.p4.integrate(info.targetWorkspace, source, changenum, integTarget, {edgeServerAddress, virtual: doVirtualMerge, nullMerge: target.mergeMode === 'null'})
 
 		const pending: PendingChange = {change: info, action: target, newCl: changenum}
 
