Delay sending nag notifications until after conflicts have been evaluated

diff --git a/Engine/Extras/RoboMerge/v3/src/robo/nodebot.ts b/Engine/Extras/RoboMerge/v3/src/robo/nodebot.ts
--- a/Engine/Extras/RoboMerge/v3/src/robo/nodebot.ts
+++ b/Engine/Extras/RoboMerge/v3/src/robo/nodebot.ts
@@ -300,14 +300,10 @@ export class NodeBot extends PerforceStatefulBot implements NodeBotInterface {
 			return false
 		}
 
-		if (this.conflicts.getConflicts().length > 0)
-		{
-			this.sendNagNotifications()
-		}
-
 		// see if our flow is blocked
 		if (this.isBlocked) {
 			this.nodeBotLogger.verbose('tick() - This bot is blocked')
+			this.sendNagNotifications()
 			await _nextTick()
 			return false
 		}
@@ -321,6 +317,7 @@ export class NodeBot extends PerforceStatefulBot implements NodeBotInterface {
 
 		if (this.skipTickCounter > 0) {
 			--this.skipTickCounter
+			this.sendNagNotifications()		
 			return false
 		}
 
@@ -378,6 +375,7 @@ export class NodeBot extends PerforceStatefulBot implements NodeBotInterface {
 				await Promise.all(Array.from(availableEdgesAtCL.entries()).map(async ([minCL,edgesAtCL]) => this._processListOfChanges(edgesAtCL, minCL, changesToProcess)));
 			}
 		}
+		this.sendNagNotifications()
 		return true
 	}
 
