Make alerts channel configurable as an environment variable

diff --git a/Engine/Extras/RoboMerge/v3/src/robo/notifications.ts b/Engine/Extras/RoboMerge/v3/src/robo/notifications.ts
--- a/Engine/Extras/RoboMerge/v3/src/robo/notifications.ts
+++ b/Engine/Extras/RoboMerge/v3/src/robo/notifications.ts
@@ -54,9 +54,8 @@ export function isUserAKnownBot(user: string) {
 }
 
 export async function postToRobomergeAlerts(message: string) {
-	// 'C044ZUSS71S' is #dt-robomerge-alert-ext
 	if (!args.devMode) {
-		return postMessageToChannel(message, 'C044ZUSS71S')
+		return postMessageToChannel(message, args.slackAlertChannel)
 	}
 }
 
diff --git a/Engine/Extras/RoboMerge/v3/src/robo/robo.ts b/Engine/Extras/RoboMerge/v3/src/robo/robo.ts
--- a/Engine/Extras/RoboMerge/v3/src/robo/robo.ts
+++ b/Engine/Extras/RoboMerge/v3/src/robo/robo.ts
@@ -219,6 +219,11 @@ const COMMAND_LINE_ARGS: {[param: string]: Arg<any>} = {
 		env: 'ROBO_SLACK_DOMAIN',
 		dflt: 'https://slack.com'
 	},
+	slackAlertChannel: {
+		match: /^-slackAlertChannel=(.+)$/,
+		env: 'ROBO_SLACK_ALERT_CHANNEL',
+		dflt: 'C044ZUSS71S' // #dt-robomerge-alert-ext
+	},
 	mongoDB_URI: {
 		match: /^-mongoDB_URI=(.+)$/,
 		env: 'ROBO_MONGODB_URI',
