Post alert when branchmap fails to parse while running
Clean up alert when branchmap fails to parse at startup
Fix await uploadFile to wait for all steps to complete
Remove unnecessary function postMessageToDefaultChannel

diff --git a/Engine/Extras/RoboMerge/v3/src/robo/autobranchupdater.ts b/Engine/Extras/RoboMerge/v3/src/robo/autobranchupdater.ts
--- a/Engine/Extras/RoboMerge/v3/src/robo/autobranchupdater.ts
+++ b/Engine/Extras/RoboMerge/v3/src/robo/autobranchupdater.ts
@@ -8,6 +8,9 @@ import { P4_FORCE, PerforceContext, Workspace } from '../common/perforce';
 import { Bot } from './bot-interfaces';
 import { BranchDefs } from './branchdefs';
 import { GraphBot } from './graphbot';
+import { postToRobomergeAlerts } from './notifications';
+import { RoboArgs } from './roboargs';
+import { SlackFile } from './slack';
 import * as p4util from '../common/p4util';
 
 const DISABLE = false
@@ -160,23 +163,28 @@ export class AutoBranchUpdater implements Bot {
 			branchGraphText = require('fs').readFileSync(`${this.workspace.directory}/${this.graphBot.filename}`, 'utf8')
 		}
 		catch (err) {
-			// @todo email author of changes!
 			this.abuLogger.printException(err, 'ERROR: failed to reload branch specs file')
+			postToRobomergeAlerts(`ERROR: failed to reload branch specs file ${this.graphBot.filename}`)
 			return
 		}
 
+		const botname = this.graphBot.branchGraph.botname
+
 		const result = await BranchDefs.parseAndValidate(this.abuLogger, branchGraphText)
 		if (!result.branchGraphDef) {
 			// @todo email author of changes!
-			let errText = 'failed to parse/validate branch specs file\n'
-			for (const error of result.validationErrors) {
-				errText += `${error}\n`
+			const errText = `Failed to parse/validate branch specs file for ${botname}`
+			const validationErrorText = result.validationErrors.join('\n')
+			const file: SlackFile = {
+				content: validationErrorText,
+				channels: RoboArgs.args.slackAlertChannel,
+				filename: "branchspecerrors.txt"
 			}
-			this.abuLogger.error(errText.trim())
+			postToRobomergeAlerts(errText, [file])
+			this.abuLogger.error(`${errText}\n${validationErrorText}`.trim())
 			return
 		}
 
-		const botname = this.graphBot.branchGraph.botname
 		if (this.graphBot.ensureStopping()) {
 			this.abuLogger.info(`Ignoring changes to ${botname} branch specs, since bot already stopping`)
 			return
diff --git a/Engine/Extras/RoboMerge/v3/src/robo/graphbot.ts b/Engine/Extras/RoboMerge/v3/src/robo/graphbot.ts
--- a/Engine/Extras/RoboMerge/v3/src/robo/graphbot.ts
+++ b/Engine/Extras/RoboMerge/v3/src/robo/graphbot.ts
@@ -16,8 +16,10 @@ import { NodeBot } from './nodebot';
 import { bindBotNotifications, BotNotifications, postMessageToChannel, postToRobomergeAlerts, SlackMessages } from './notifications';
 import { PersistedSlackMessages } from './persistedmessages';
 import { roboAnalytics } from './roboanalytics';
+import { RoboArgs } from './roboargs';
 import { BlockageNodeOpUrls, OperationUrlHelper } from './roboserver';
 import { Settings } from './settings';
+import { SlackFile } from './slack';
 import { Status } from './status';
 import { GraphBotState } from "./status-types"
 import { TickJournal } from './tick-journal';
@@ -67,7 +69,12 @@ export class GraphBot implements BotEventHandler {
 		const result = await BranchDefs.parseAndValidate(graphBot.botLogger, fileText)
 		if (!result.branchGraphDef) {
 			const parseError = result.validationErrors.length === 0 ? 'Failed to parse' : result.validationErrors.join('\n')
-			await postToRobomergeAlerts(`@here Failed to create bot ${botname.toUpperCase}. Halting service startup.\n${parseError}`)
+			const file: SlackFile = {
+				content: parseError,
+				channels: RoboArgs.args.slackAlertChannel,
+				filename: "branchspecerrors.txt"
+			}
+			await postToRobomergeAlerts(`@here Failed to create bot ${botname.toUpperCase()}.\nHalting service startup.`, [file])
 			throw new Error(parseError)
 		}
 
diff --git a/Engine/Extras/RoboMerge/v3/src/robo/notifications.ts b/Engine/Extras/RoboMerge/v3/src/robo/notifications.ts
--- a/Engine/Extras/RoboMerge/v3/src/robo/notifications.ts
+++ b/Engine/Extras/RoboMerge/v3/src/robo/notifications.ts
@@ -11,6 +11,7 @@ import { PersistentConflict, Resolution } from './conflict-interfaces';
 import { BotEventHandler, BotEvents } from './events';
 import { NodeBot } from './nodebot';
 import { PersistedSlackMessage, PersistedSlackMessages } from './persistedmessages';
+import { RoboArgs } from './roboargs';
 import { BlockageNodeOpUrls } from './roboserver';
 import { Slack, SlackAttachment, SlackFile, SlackLinkButtonAction, SlackLinkButtonsAttachment, SlackMessageField, SlackMessage, SlackMessageStyles } from './slack';
 import { Vault } from './vault'
@@ -53,20 +54,28 @@ export function isUserAKnownBot(user: string) {
 	return KNOWN_BOT_NAMES.indexOf(user) !== -1 || KNOWN_BOT_EMAILS.indexOf(user) !== -1
 }
 
-export async function postToRobomergeAlerts(message: string) {
+export async function postToRobomergeAlerts(message: string, replies?: (string|SlackFile|SlackMessage)[]) {
 	if (!args.devMode) {
-		return postMessageToChannel(message, args.slackAlertChannel)
+		return postMessageToChannel(`${message}\n\n${RoboArgs.args.externalUrl}`, args.slackAlertChannel, SlackMessageStyles.DANGER, replies)
 	}
 }
 
-export async function postMessageToChannel(message: string, channel: string, style: SlackMessageStyles = SlackMessageStyles.GOOD) {
+export async function postMessageToChannel(message: string, channel: string, style: SlackMessageStyles = SlackMessageStyles.GOOD, replies?: (string|SlackFile|SlackMessage)[]) {
 	if (SLACK_TOKENS.bot) {
-		return new Slack({id: channel, botToken: SLACK_TOKENS.bot, userToken: SLACK_TOKENS.user}, args.slackDomain,new ContextualLogger(`Post Message to ${channel}`)).postMessageToDefaultChannel({
-			text: message,
-			style,
-			channel,
-			mrkdwn: true
-		})
+		const slack = new Slack({id: channel, botToken: SLACK_TOKENS.bot, userToken: SLACK_TOKENS.user}, args.slackDomain, new ContextualLogger(`Post Message to ${channel}`))	
+		const thread_ts = await slack.postMessage({text: message, style, channel, mrkdwn: true })
+		for (const reply of (replies || [])) {
+			if (typeof reply === 'string') {
+				await slack.reply(thread_ts, {text: message, channel, mrkdwn: true})
+			}
+			else if ('content' in reply) {
+				reply.thread_ts = thread_ts
+				await slack.uploadFile(reply)
+			}
+			else {
+				await slack.reply(thread_ts, reply)
+			}
+		}
 	}
 }
 
diff --git a/Engine/Extras/RoboMerge/v3/src/robo/slack.ts b/Engine/Extras/RoboMerge/v3/src/robo/slack.ts
--- a/Engine/Extras/RoboMerge/v3/src/robo/slack.ts
+++ b/Engine/Extras/RoboMerge/v3/src/robo/slack.ts
@@ -109,12 +109,6 @@ export class Slack {
 		return (await this.post('chat.postMessage', this.makeArgs(message))).ts
 	}
 
-	postMessageToDefaultChannel(message: SlackMessage) {
-		const args = this.makeArgs(message)
-		args.channel = this.channel.id
-		return this.post('chat.postMessage', args)
-	}
-
 	reply(thread_ts: string, message: SlackMessage) {
 		const args = this.makeArgs(message)
 		args.thread_ts = thread_ts
@@ -137,7 +131,7 @@ export class Slack {
 		}
 		return this.post('files.getUploadURLExternal', getURLargs, true, PostFormat.URL)
 			.then( getURLresult => {
-				if (getURLresult.ok) {
+				return (getURLresult.ok ?
 					request.post({
 						url: getURLresult.upload_url,
 						body: file.content,
@@ -150,10 +144,10 @@ export class Slack {
 								initial_comment: file.initial_comment,
 								thread_ts: file.thread_ts
 							}
-							this.post('files.completeUploadExternal', completeArgs)
+							return this.post('files.completeUploadExternal', completeArgs)
 						}
 					)
-				}
+					: undefined)
 			}
 		)
 	}
