Broadcast pause/unpause messages to appropriate channels

diff --git a/Engine/Extras/RoboMerge/v3/src/robo/branch-interfaces.ts b/Engine/Extras/RoboMerge/v3/src/robo/branch-interfaces.ts
--- a/Engine/Extras/RoboMerge/v3/src/robo/branch-interfaces.ts
+++ b/Engine/Extras/RoboMerge/v3/src/robo/branch-interfaces.ts
@@ -362,4 +362,12 @@ export type NagBlockageEvent = {
 	triager: string
 	triagerEmail: Promise<string>
 	ownerEmail: Promise<string>
+}
+
+export type PauseEvent = {
+	nodeOrEdgeName: string
+	sourceBranchUpperName: string
+	targetBranchUpperName?: string
+	message?: string
+	owner: string
 }
\ No newline at end of file
diff --git a/Engine/Extras/RoboMerge/v3/src/robo/edgebot.ts b/Engine/Extras/RoboMerge/v3/src/robo/edgebot.ts
--- a/Engine/Extras/RoboMerge/v3/src/robo/edgebot.ts
+++ b/Engine/Extras/RoboMerge/v3/src/robo/edgebot.ts
@@ -21,7 +21,7 @@ import { RoboArgs } from "./roboargs";
 import { Context } from "./settings";
 import { SlackMessageStyles } from "./slack";
 import { PauseState } from "./state-interfaces";
-import { BlockagePauseInfo, BlockagePauseInfoMinimal, EdgeStatusFields } from "./status-types";
+import { BlockagePauseInfo, EdgeStatusFields } from "./status-types";
 import { getIntegrationOwner } from "./targets";
 
 const FAILED_CHANGELIST_PAUSE_TIMEOUT_SECONDS = 15 * 60
@@ -237,10 +237,6 @@ class EdgeBotImpl extends PerforceStatefulBot {
 		}
 	}
 
-	block(blockageInfo: BlockagePauseInfoMinimal, pauseDurationSeconds?: number) {
-		super.block(blockageInfo, pauseDurationSeconds)
-	}
-
 	private analyzeConflict(unresolved: ConflictedResolveNFile[]) {
 		const results: ConflictingFile[] = []
 
@@ -457,6 +453,34 @@ class EdgeBotImpl extends PerforceStatefulBot {
 		this.gate.persist()
 	}
 
+	pause(message: string, owner: string) {
+		if (this.isManuallyPaused) {
+			return
+		}
+
+		super.pause(message, owner)
+		this.eventTriggers.reportPause({message, owner, 
+			nodeOrEdgeName: this.displayName,
+			sourceBranchUpperName: this.sourceBranch.upperName, 
+			targetBranchUpperName: this.branch.upperName
+		})
+	}
+
+	unpause(requester: string) {
+		if (!this.isManuallyPaused) {
+			return
+		}
+
+		super.unpause(requester)
+		this.eventTriggers.reportUnpause({
+			message: "",
+			owner: requester, 
+			nodeOrEdgeName: this.displayName,
+			sourceBranchUpperName: this.sourceBranch.upperName, 
+			targetBranchUpperName: this.branch.upperName
+		})
+	}
+
 	queueUnblock() {
 		this.sourceNode.queueEdgeUnblock(this.targetBranch.upperName)
 	}
diff --git a/Engine/Extras/RoboMerge/v3/src/robo/events.ts b/Engine/Extras/RoboMerge/v3/src/robo/events.ts
--- a/Engine/Extras/RoboMerge/v3/src/robo/events.ts
+++ b/Engine/Extras/RoboMerge/v3/src/robo/events.ts
@@ -1,6 +1,6 @@
 // Copyright Epic Games, Inc. All Rights Reserved.
 
-import { AlreadyIntegrated, Blockage, ChangeInfo, ForcedCl, PendingChange } from './branch-interfaces';
+import { AlreadyIntegrated, Blockage, ChangeInfo, ForcedCl, PauseEvent, PendingChange } from './branch-interfaces';
 import { BeginIntegratingToGateEvent, EndIntegratingToGateEvent, NagBlockageEvent } from './branch-interfaces';
 import { BotConfig } from './branchdefs';
 import { PersistentConflict } from './conflict-interfaces';
@@ -13,6 +13,8 @@ type OnCommit = (arg: PendingChange) => void
 type OnAlreadyIntegrated = (arg: AlreadyIntegrated) => void
 type OnForcedLastCl = (arg: ForcedCl) => void
 type OnNonSkipLastClChange = (arg: ForcedCl) => void
+type OnPause = (arg: PauseEvent) => void
+type OnUnpause = (arg: PauseEvent) => void
 type OnBranchUnblocked = (arg: PersistentConflict) => void
 type OnConflictStatus = (arg: boolean) => void
 type OnBeginIntegratingToGate = (arg: BeginIntegratingToGateEvent) => void
@@ -26,6 +28,8 @@ export interface BotEventHandler {
 	onAlreadyIntegrated?: OnAlreadyIntegrated
 	onForcedLastCl?: OnForcedLastCl
 	onNonSkipLastClChange?: OnNonSkipLastClChange
+	onPause?: OnPause
+	onUnpause?: OnUnpause
 	onBranchUnblocked?: OnBranchUnblocked
 	onConflictStatus?: OnConflictStatus
 
@@ -77,6 +81,14 @@ export class BotEvents {
 		this.nonSkipLastClChangeListeners.push(listener)
 	}
 
+	onPause(listener: OnPause) {
+		this.pauseListeners.push(listener)
+	}
+
+	onUnpause(listener: OnUnpause) {
+		this.unpauseListeners.push(listener)
+	}
+
 	// Fired when a previously blocked bot has been unblocked
 	onBranchUnblocked(listener: OnBranchUnblocked) {
 		this.branchUnblockedListeners.push(listener)
@@ -109,6 +121,8 @@ export class BotEvents {
 				case 'onAlreadyIntegrated': this.alreadyIntegratedListeners.push(arg => prop.call(handler, arg)); break
 				case 'onForcedLastCl': this.forcedLastClListeners.push(arg => prop.call(handler, arg)); break
 				case 'onNonSkipLastClChange': this.nonSkipLastClChangeListeners.push(arg => prop.call(handler, arg)); break
+				case 'onPause': this.pauseListeners.push(arg => prop.call(handler, arg)); break
+				case 'onUnpause': this.unpauseListeners.push(arg => prop.call(handler, arg)); break
 				case 'onBranchUnblocked': this.branchUnblockedListeners.push(arg => prop.call(handler, arg)); break
 				case 'onConflictStatus': this.conflictStatusListeners.push(arg => prop.call(handler, arg)); break
 				case 'onBeginIntegratingToGate': this.beginIntegratingToGateListeners.push(arg => prop.call(handler, arg)); break
@@ -125,6 +139,8 @@ export class BotEvents {
 	protected readonly alreadyIntegratedListeners: OnAlreadyIntegrated[] = []
 	protected readonly forcedLastClListeners: OnForcedLastCl[] = []
 	protected readonly nonSkipLastClChangeListeners: OnNonSkipLastClChange[] = []
+	protected readonly pauseListeners: OnPause[] = []
+	protected readonly unpauseListeners: OnUnpause[] = []
 	protected readonly branchUnblockedListeners: OnBranchUnblocked[] = []
 	protected readonly conflictStatusListeners: OnConflictStatus[] = []
 	protected readonly beginIntegratingToGateListeners: OnBeginIntegratingToGate[] = []
@@ -181,6 +197,18 @@ export class BotEventTriggers extends BotEvents {
 		}
 	}
 
+	reportPause(arg: PauseEvent) {
+		for (const listener of this.pauseListeners) {
+			listener(arg)
+		}
+	}
+
+	reportUnpause(arg: PauseEvent) {
+		for (const listener of this.unpauseListeners) {
+			listener(arg)
+		}
+	}
+
 	reportBranchUnblocked(arg: PersistentConflict) {
 		for (const listener of this.branchUnblockedListeners) {
 			listener(arg)
diff --git a/Engine/Extras/RoboMerge/v3/src/robo/nodebot.ts b/Engine/Extras/RoboMerge/v3/src/robo/nodebot.ts
--- a/Engine/Extras/RoboMerge/v3/src/robo/nodebot.ts
+++ b/Engine/Extras/RoboMerge/v3/src/robo/nodebot.ts
@@ -119,7 +119,6 @@ export class NodeBot extends PerforceStatefulBot implements NodeBotInterface {
 		const settingsQueuedUnblocks = this.settings.get('edgesQueuedToUnblock')
 		this.edgesQueuedToUnblock = new Set(Array.isArray(settingsQueuedUnblocks) ? settingsQueuedUnblocks : [])
 		
-		this.eventTriggers = eventTriggers
 		this.conflicts = new Conflicts(branchDef.upperName, eventTriggers, this.settings, externalUrl, this.branchGraph.config.reportToBuildHealth, this.nodeBotLogger)
 
 		this.eventTriggers.onChangeParsed((info: ChangeInfo) => this.onGlobalChange(info))
@@ -538,6 +537,35 @@ export class NodeBot extends PerforceStatefulBot implements NodeBotInterface {
 		return this.nodeBotLogger
 	}
 
+	pause(message: string, owner: string) {
+		if (this.isManuallyPaused) {
+			return
+		}
+
+		super.pause(message, owner)
+		this.eventTriggers.reportPause({ 
+			message,
+			owner, 
+			nodeOrEdgeName: this.displayName,
+			sourceBranchUpperName: this.branch.upperName, 
+			targetBranchUpperName: this.branch.upperName
+		})
+	}
+
+	unpause(requester: string) {
+		if (!this.isManuallyPaused) {
+			return
+		}
+
+		super.unpause(requester)
+		this.eventTriggers.reportUnpause({
+			owner: requester, 
+			nodeOrEdgeName: this.displayName,
+			sourceBranchUpperName: this.branch.upperName, 
+			targetBranchUpperName: this.branch.upperName
+		})
+	}
+
 	forceSetLastClWithContext(value: number, culprit: string, reason: string, _?: boolean) {
 		const prevValue = this._forceSetLastCl_NoReset(value)
 
diff --git a/Engine/Extras/RoboMerge/v3/src/robo/notifications.ts b/Engine/Extras/RoboMerge/v3/src/robo/notifications.ts
--- a/Engine/Extras/RoboMerge/v3/src/robo/notifications.ts
+++ b/Engine/Extras/RoboMerge/v3/src/robo/notifications.ts
@@ -6,7 +6,7 @@ import { Badge } from '../common/badge';
 import { Random, setDefault } from '../common/helper';
 import { ContextualLogger } from '../common/logger';
 import { ExclusiveFileDetails } from '../common/perforce';
-import { Blockage, Branch, BranchArg, EditableBranch, ExclusiveLockInfo, ForcedCl, MergeAction, NagBlockageEvent, NodeOpUrlGenerator, resolveBranchArg } from './branch-interfaces';
+import { Blockage, Branch, BranchArg, EditableBranch, ExclusiveLockInfo, ForcedCl, MergeAction, NagBlockageEvent, NodeOpUrlGenerator, PauseEvent, resolveBranchArg } from './branch-interfaces';
 import { PersistentConflict, Resolution } from './conflict-interfaces';
 import { BotEventHandler, BotEvents } from './events';
 import { NodeBot } from './nodebot';
@@ -817,6 +817,44 @@ export class BotNotifications implements BotEventHandler {
 		}
 	}
 
+	onPause(details: PauseEvent) {
+		if (this.slackMessages) {
+			const channelsToPostTo = this.getChannelsToPostTo(details.sourceBranchUpperName, details.targetBranchUpperName)
+			for (let slackChannel of channelsToPostTo) {
+				this.slackMessages.postNonConflictMessage({
+					title: details.nodeOrEdgeName + ' paused',
+					text: details.message || "",
+					style: SlackMessageStyles.WARNING,
+					fields: [{
+						title: 'By', short: true, value: details.owner
+					}],
+					title_link: this.externalRobomergeUrl + '#' + this.botname,
+					mrkdwn: true,
+					channel: slackChannel // Default to the configured channel
+				}) 
+			}
+		}
+	}
+
+	onUnpause(details: PauseEvent) {
+		if (this.slackMessages) {
+			const channelsToPostTo = this.getChannelsToPostTo(details.sourceBranchUpperName, details.targetBranchUpperName)
+			for (let slackChannel of channelsToPostTo) {
+				this.slackMessages.postNonConflictMessage({
+					title: details.nodeOrEdgeName + ' unpaused',
+					text: details.message || "",
+					style: SlackMessageStyles.WARNING,
+					fields: [{
+						title: 'By', short: true, value: details.owner
+					}],
+					title_link: this.externalRobomergeUrl + '#' + this.botname,
+					mrkdwn: true,
+					channel: slackChannel // Default to the configured channel
+				}) 
+			}
+		}
+	}
+
 	sendTestMessage(username : string) {
 
 		if (this.slackMessages) {
