New Changelist Filter feature for edges allows extensible filters to determine which changes are allowed and which will be blocked/skipped. Initially provided are a perforce group implementation and a lockdown allowlist implementation.

diff --git a/Engine/Extras/RoboMerge/v3/package.json b/Engine/Extras/RoboMerge/v3/package.json
--- a/Engine/Extras/RoboMerge/v3/package.json
+++ b/Engine/Extras/RoboMerge/v3/package.json
@@ -1,6 +1,6 @@
 {
   "name": "robomerge",
-  "version": "3.3.0",
+  "version": "3.4.0",
   "description": "Automatic P4 merge management bot.",
   "repository": {},
   "keywords": [
@@ -24,6 +24,7 @@
     "@types/swagger-ui-dist": "^3.0.3",
     "@viz-js/viz": "^3.11.0",
     "colors": "^1.4.0",
+    "ini": "^6.0.0",
     "jsonlint-mod": "^1.7.6",
     "ldapauth-fork": "^4.0.2",
     "markdown-toc": "^1.2.0",
@@ -35,6 +36,7 @@
     "sanitize-html": "^2.5.3",
     "semver": "^6.1.0",
     "source-map-support": "^0.5.12",
+    "typescript": "^5.9.3",
     "winston": "^3.2.1"
   },
   "devDependencies": {
diff --git a/Engine/Extras/RoboMerge/v3/src/common/perforce.ts b/Engine/Extras/RoboMerge/v3/src/common/perforce.ts
--- a/Engine/Extras/RoboMerge/v3/src/common/perforce.ts
+++ b/Engine/Extras/RoboMerge/v3/src/common/perforce.ts
@@ -1783,6 +1783,17 @@ export class PerforceContext {
 		return this._execP4(roboWorkspace, ["shelve", "-d", "-c", changelist.toString(), "-f"]);
 	}
 
+	getGroup(groupname: string) {
+		// Verify it is actually a group, then return the details of it
+		return this.execAndParse(null, ['groups', '-v', groupname], {quiet:true})
+					.then(groups => {
+						if (groups.length == 0) {
+							return null
+						}
+						return this.execAndParseArray(null, ['group', '-o', groupname], {quiet:true})
+					})
+	}
+
 	async getUser(username: string) {
 		if (username.length > 0 && !isUserSlackGroup(username) && !isUserAKnownBot(username)) {
 			try {
diff --git a/Engine/Extras/RoboMerge/v3/src/robo/branchdefs.ts b/Engine/Extras/RoboMerge/v3/src/robo/branchdefs.ts
--- a/Engine/Extras/RoboMerge/v3/src/robo/branchdefs.ts
+++ b/Engine/Extras/RoboMerge/v3/src/robo/branchdefs.ts
@@ -3,6 +3,7 @@
 import { setDifference, setIntersection } from '../common/helper'
 import { PerforceContext } from '../common/perforce'
 import { ContextualLogger } from '../common/logger'
+import { changelistFilters } from './targets'
 
 const jsonlint: any = require('jsonlint-mod')
 
@@ -74,7 +75,7 @@ export class IntegrationMethod {
 }
 
 export const DAYS_OF_THE_WEEK = ['sun', 'mon', 'tue', 'wed', 'thu', 'fri', 'sat']
-type DayOfTheWeek = 'sun' | 'mon' | 'tue' | 'wed' | 'thu' | 'fri' | 'sat'
+type DayOfTheWeek = typeof DAYS_OF_THE_WEEK[number]
 
 export type IntegrationWindowPane = {
 	// if days not specified, daily
@@ -170,6 +171,12 @@ const edgeOptionFieldsPrototype = {
 		description: '',
 		channelId: '',
 		block: true
+	},
+
+	changelistFilter: {
+		name: '',
+		mode: '',
+		parameters: []
 	}
 }
 
@@ -178,6 +185,22 @@ const edgeOptionFieldNames: ReadonlySet<string> = new Set(Object.keys(edgeOption
 type EdgeOptionFields = typeof edgeOptionFieldsPrototype
 export type ApprovalOptions = typeof edgeOptionFieldsPrototype.approval
 
+const FILTER_MODES = ['AllowOrBlock', 'AllowOrSkip'] as const
+export type FilterModeType = typeof FILTER_MODES[number]
+
+function toFilterMode(value: string): FilterModeType {
+  if ((FILTER_MODES as readonly string[]).includes(value)) {
+    return value as FilterModeType;
+  }
+  throw new Error(`Invalid group member type: ${value}`);
+}
+
+export type FilterOptions = {
+	name: string,
+	mode: FilterModeType,
+	parameters: string[]
+}
+
 export type NodeOptions = Partial<NodeOptionFields>
 export type EdgeOptions = Partial<EdgeOptionFields>
 
@@ -286,7 +309,7 @@ export class BranchDefs {
 		}
 	}
 
-	static async parseAndValidate(logger: ContextualLogger, branchSpecsText: string, isPreviewing?: boolean): Promise<ParseResult> {
+	static async parseAndValidate(logger: ContextualLogger, branchSpecsText: string, isPreviewing: boolean = false): Promise<ParseResult> {
 		const defaultConfigForWholeBot: BotConfig = {
 			defaultStreamDepot: null,
 			isDefaultBot: false,
@@ -488,7 +511,31 @@ export class BranchDefs {
 					}
 				}
 
-				if (edge.integrationMethod) {
+			if (edge.changelistFilter) {
+				try {
+					await import(`./changelistfilters/${edge.changelistFilter.name}`)
+					const changelistFilter = changelistFilters.get(edge.changelistFilter.name)
+					if (changelistFilter) {
+						const [cfErrors, cfWarnings] = await changelistFilter.validate(branchGraph.branches, edge, isPreviewing, logger)
+						cfErrors.forEach(err => validationErrors.push(`Custom filter '${edge.changelistFilter!.name}' for edge ${edge.from}->${edge.to}: ${err}`))
+						cfWarnings.forEach(warn => validationWarnings.push(`Custom filter '${edge.changelistFilter!.name}' for edge ${edge.from}->${edge.to}: ${warn}`))
+					}
+					else {
+						validationErrors.push(`Edge ${edge.from}->${edge.to} specifies custom filter '${edge.changelistFilter.name}' but it is not found after import`)
+					}
+					try {
+						toFilterMode(edge.changelistFilter.mode)
+					}
+					catch {
+						validationErrors.push(`Edge ${edge.from}->${edge.to} specifies invalid filter mode '${edge.changelistFilter.mode}'`)
+					}
+				}
+				catch {
+					validationErrors.push(`Edge ${edge.from}->${edge.to} specifies custom filter '${edge.changelistFilter.name}' but unable to import it`)
+				}
+			}
+
+			if (edge.integrationMethod) {
 					BranchDefs.checkValidIntegrationMethod(validationErrors, edge.integrationMethod, `${edge.from}->${edge.to}`)
 				}
 
diff --git a/Engine/Extras/RoboMerge/v3/src/robo/changelistfilters/allowlist.ts b/Engine/Extras/RoboMerge/v3/src/robo/changelistfilters/allowlist.ts
new file mode 100644
--- /dev/null
+++ b/Engine/Extras/RoboMerge/v3/src/robo/changelistfilters/allowlist.ts
@@ -0,0 +1,90 @@
+// Copyright Epic Games, Inc. All Rights Reserved.
+
+import { ContextualLogger } from 'common/logger'
+import { PerforceContext } from '../../common/perforce'
+import { ChangeInfo, MergeAction } from '../branch-interfaces'
+import { EdgeOptions, EdgeProperties, NodeOptions } from '../branchdefs'
+import { ChangelistFilter, ChangelistFilterResult, changelistFilters } from '../targets'
+const ini = require('ini')
+
+export class AllowListFilter extends ChangelistFilter {
+
+	static {
+		changelistFilters.set("allowlist", AllowListFilter)
+	}
+
+	static async shouldFilter(info: ChangeInfo, action: MergeAction, edgeProperties: EdgeOptions, logger: ContextualLogger) {
+
+        const authorizedUser = info.userRequest === 'edge-reconsider' ? info.owner || info.author : info.author
+        const authorizedUserLower = authorizedUser.toLowerCase()
+		const p4 = PerforceContext.getServerContext(logger, action.branch.config.streamServer)
+        const configFile = await p4.print(edgeProperties.changelistFilter!.parameters[0])
+        const config = ini.parse(configFile);
+        let result: ChangelistFilterResult = {result: 'allow'}
+
+        if (!config[edgeProperties.changelistFilter!.parameters[1]].allowlist?.split(',').some((user: string) => user.trim().toLowerCase() === authorizedUserLower))
+        {
+            if (edgeProperties.changelistFilter?.mode === 'AllowOrBlock') {
+                result.result = 'block'
+                result.message = `${authorizedUser} not in allowlist of [${edgeProperties.changelistFilter!.parameters[1]}]`
+            }
+            else {
+                result.result = 'skip'
+            }
+        }
+
+        return result
+	}
+
+    static async validate(branches: NodeOptions[], edge: EdgeProperties, isPreviewing: boolean, logger: ContextualLogger): Promise<[string[],string[]]> {
+
+        const errors: string[] = []
+        const warnings: string[] = []
+
+        const warnWhenPreviewing = (message: string) => {
+            if (isPreviewing) {
+                warnings.push(message)
+            }
+            else {
+                errors.push(message)
+            }
+        }
+
+        if (edge.changelistFilter!.parameters.length != 2) {
+            errors.push("Expected exactly 2 parameters: perforcePath, sectionName")
+        }
+
+        const edgeToUpper = edge.to.toUpperCase()
+        const toServer = branches.filter(branch => branch.name?.toUpperCase() === edgeToUpper)[0].streamServer
+        const p4 = PerforceContext.getServerContext(logger, toServer)
+        
+        if (edge.changelistFilter!.parameters.length > 0) {            
+            const configFile = await p4.print(edge.changelistFilter!.parameters[0])
+            if (configFile.length > 0) {
+                try {
+                    const config = ini.parse(configFile);
+                    if (edge.changelistFilter!.parameters.length > 1) {
+                        const section = config[edge.changelistFilter!.parameters[1]]
+                        if (section) {
+                            if (!section.allowlist) {
+                                warnings.push(`[${edge.changelistFilter!.parameters[1]}] does not have allowlist in ${edge.changelistFilter!.parameters[0]}`)
+                            }
+                        }
+                        else {
+                            warnWhenPreviewing(`[${edge.changelistFilter!.parameters[1]}] not found in ${edge.changelistFilter!.parameters[0]}`)
+                        }
+                    }
+                }
+                catch {
+                    errors.push(`Failed to parse ${edge.changelistFilter!.parameters[0]}`)
+                }
+            }
+            else {
+                warnWhenPreviewing(`${edge.changelistFilter!.parameters[0]} does not exist.`)
+            }
+        }
+
+        return [errors, warnings]
+    }
+
+}
\ No newline at end of file
diff --git a/Engine/Extras/RoboMerge/v3/src/robo/changelistfilters/perforcegroup.ts b/Engine/Extras/RoboMerge/v3/src/robo/changelistfilters/perforcegroup.ts
new file mode 100644
--- /dev/null
+++ b/Engine/Extras/RoboMerge/v3/src/robo/changelistfilters/perforcegroup.ts
@@ -0,0 +1,120 @@
+// Copyright Epic Games, Inc. All Rights Reserved.
+
+import { ContextualLogger } from 'common/logger'
+import { PerforceContext } from '../../common/perforce'
+import { ChangeInfo, MergeAction } from '../branch-interfaces'
+import { EdgeOptions, EdgeProperties, NodeOptions } from '../branchdefs'
+import { ChangelistFilter, ChangelistFilterResult, changelistFilters } from '../targets'
+
+const GROUP_MEMBER_TYPES = ['owners', 'users', 'any'] as const
+type GroupMemberType = typeof GROUP_MEMBER_TYPES[number]
+
+function toGroupMemberType(value: string): GroupMemberType {
+  if ((GROUP_MEMBER_TYPES as readonly string[]).includes(value)) {
+    return value as GroupMemberType;
+  }
+  throw new Error(`Invalid group member type: ${value}`);
+}
+
+export class PerforceGroupFilter extends ChangelistFilter {
+
+	static {
+		changelistFilters.set("perforcegroup", PerforceGroupFilter)
+	}
+
+	static async shouldFilter(info: ChangeInfo, action: MergeAction, edgeProperties: EdgeOptions, logger: ContextualLogger) {
+
+		const p4 = PerforceContext.getServerContext(logger, action.branch.config.streamServer)
+
+        const authorizedUser = info.userRequest === 'edge-reconsider' ? info.owner || info.author : info.author
+        const memberType: GroupMemberType = toGroupMemberType(edgeProperties.changelistFilter!.parameters[1])
+
+        let result: ChangelistFilterResult = {result: 'allow'}
+
+        let groups = [p4.getGroup(edgeProperties.changelistFilter!.parameters[0])]
+        for (const groupPromise of groups) {
+            const group = (await groupPromise)?.[0]
+            if (group && group.entries) {
+                switch(memberType) {
+                    case 'owners':
+                        if (group.entries.some((entry: any) => entry.Owners === authorizedUser)) {
+                            return result
+                        }
+                        break
+                    case 'users':
+                        if (group.entries.some((entry: any) => entry.Users === authorizedUser)) {
+                            return result
+                        }
+                        break
+                    case 'any':
+                        if (group.entries.some((entry: any) => entry.Owners === authorizedUser || entry.Users === authorizedUser)) {
+                            return result
+                        }
+                        break
+                }
+                group.entries.forEach((entry: any) => { if (entry.Subgroups) { groups.push(p4.getGroup(entry.Subgroups)) }})
+            }
+        }
+
+        const getMemberType = function() {
+            switch(memberType) {
+                case 'owners':
+                    return 'an owner of'
+                case 'users':
+                    return 'a user of'
+                case 'any':
+                    return 'in'
+            }
+        }
+
+        if (edgeProperties.changelistFilter?.mode === 'AllowOrBlock') {
+            result.result = 'block'
+            result.message = `${authorizedUser} not ${getMemberType()} group ${edgeProperties.changelistFilter!.parameters[0]}`
+        }
+        else {
+            result.result = 'skip'
+        }
+		return result
+	}
+
+    static async validate(branches: NodeOptions[], edge: EdgeProperties, isPreviewing: boolean, logger: ContextualLogger): Promise<[string[],string[]]> {
+
+        const errors: string[] = []
+        const warnings: string[] = []
+
+        const warnWhenPreviewing = (message: string) => {
+            if (isPreviewing) {
+                warnings.push(message)
+            }
+            else {
+                errors.push(message)
+            }
+        }
+
+        if (edge.changelistFilter!.parameters.length != 2) {
+            errors.push("Expected exactly 2 parameters: groupName, groupMemberType")
+        }
+
+        if (edge.changelistFilter!.parameters.length > 1) {
+            try {
+                toGroupMemberType(edge.changelistFilter!.parameters[1])
+            } 
+            catch {
+                errors.push("Parameter 2 is expected to be one of 'owners', 'users', or 'any'")
+            }
+        }
+
+        if (edge.changelistFilter!.parameters.length > 0) {
+            const edgeToUpper = edge.to.toUpperCase()
+            const toServer = branches.filter(branch => branch.name?.toUpperCase() === edgeToUpper)[0].streamServer
+            const p4 = PerforceContext.getServerContext(logger, toServer)
+            
+            if (!(await p4.getGroup(edge.changelistFilter!.parameters[0]))) {
+                warnWhenPreviewing(`Group ${edge.changelistFilter!.parameters[0]} does not exist.`)
+            }
+        }
+
+        return [errors, warnings]
+    }
+
+}
\ No newline at end of file
diff --git a/Engine/Extras/RoboMerge/v3/src/robo/edgebot.ts b/Engine/Extras/RoboMerge/v3/src/robo/edgebot.ts
--- a/Engine/Extras/RoboMerge/v3/src/robo/edgebot.ts
+++ b/Engine/Extras/RoboMerge/v3/src/robo/edgebot.ts
@@ -865,7 +865,7 @@ class EdgeBotImpl extends PerforceStatefulBot {
 				}
 			}
 
-			if (this.options.approval) {
+			if (this.options.approval && !pending.action.flags.has('#filter_approved#')) {
 				const approval = this.options.approval
 
 				// shelve
@@ -888,6 +888,14 @@ class EdgeBotImpl extends PerforceStatefulBot {
 						});
 				}
 			}
+			else if (pending.action.flags.has('#filter_blocked#')) {
+
+				const blockedMessage = [...pending.action.flags].filter(flag => flag.startsWith('#filter_reason_')).map(flag => flag.slice('#filter_reason_'.length))?.[0]
+										|| "Blocked by filter" 
+
+				// still notify as a 'failure' to trigger normal blockage mechanism
+				failure = { kind: 'Blocked', description: blockedMessage, details: blockedMessage }
+			}
 			else {
 				// no conflicts, try to submit
 				const result = await this.submitChangelist(pending, submitRetries)
diff --git a/Engine/Extras/RoboMerge/v3/src/robo/nodebot.ts b/Engine/Extras/RoboMerge/v3/src/robo/nodebot.ts
--- a/Engine/Extras/RoboMerge/v3/src/robo/nodebot.ts
+++ b/Engine/Extras/RoboMerge/v3/src/robo/nodebot.ts
@@ -24,7 +24,7 @@ import { Context } from './settings';
 import { SlackFile, SlackMessage, SlackMessageField, SlackMessageStyles } from './slack';
 import { PauseState } from './state-interfaces';
 import { newTickJournal, TickJournal } from './tick-journal';
-import { computeTargets, descriptionParsers, parseDescription, processOtherBotTargets, getIntegrationOwner, getNodeBotFullName, getNodeBotFullNameForLogging } from './targets';
+import { changelistFilters, computeTargets, descriptionParsers, parseDescription, processOtherBotTargets, getIntegrationOwner, getNodeBotFullName, getNodeBotFullNameForLogging } from './targets';
 import { GraphAPI } from './graph';
 
 /**********************************
@@ -1641,7 +1641,7 @@ export class NodeBot extends PerforceStatefulBot implements NodeBotInterface {
 		this._log_action(`Analyzing CL ${change.change}`, "verbose")
 
 		// parse the change
-		const info : ChangeInfo = this.parseChange(change, optTargetBranch)
+		const info : ChangeInfo = await this.parseChange(change, optTargetBranch)
 
 		if (optOwnerOverride) {
 			info.owner = optOwnerOverride
@@ -1808,7 +1808,7 @@ export class NodeBot extends PerforceStatefulBot implements NodeBotInterface {
 		this._sendGenericEmail(recipients, subject, errorPreface, message);
 	}
 
-	private parseChange(change: Change, optTargetBranch?: Branch | null): ChangeInfo {
+	private async parseChange(change: Change, optTargetBranch?: Branch | null): Promise<ChangeInfo> {
 
 		// parse the description
 		if (change.desc && typeof(change.desc.split) !== 'function') {
@@ -1923,17 +1923,15 @@ export class NodeBot extends PerforceStatefulBot implements NodeBotInterface {
 				const [computedMergeActions, flags] = computeResult
 
 				if (computedMergeActions.length > 0) {
-					if (info.userRequest) {
-						info.targets = computedMergeActions
-					}
-					else {
-						const mergeActions: MergeAction[] = []
 
-						for (const action of computedMergeActions) {
-							const edgeProperties = this.branch.edgeProperties.get(action.branch.upperName)
+					const mergeActions: MergeAction[] = []
+
+					for (const action of computedMergeActions) {
+						const edgeProperties = this.branch.edgeProperties.get(action.branch.upperName)
 
+						if (edgeProperties) {
 							// let ignore and disregard cancel each other out here, but don't honour ignore here so we can validate it at the bottom of the function
-							const authorIsExcluded = edgeProperties && edgeProperties.excludeAuthors && edgeProperties.excludeAuthors.indexOf(info.author) >= 0
+							const authorIsExcluded = !info.userRequest && edgeProperties.excludeAuthors && edgeProperties.excludeAuthors.indexOf(info.author) >= 0
 
 							if (authorIsExcluded && (!flags.has('disregardexcludedauthors') || flags.has('ignore'))) {
 								// skip this one
@@ -1942,8 +1940,8 @@ export class NodeBot extends PerforceStatefulBot implements NodeBotInterface {
 								continue
 							}
 
-							let descriptionIsExcluded = false
-							if (edgeProperties && edgeProperties.excludeDescriptions) {
+							if (edgeProperties.excludeDescriptions && !info.userRequest) {
+								let descriptionIsExcluded = false
 								for (const pattern of edgeProperties.excludeDescriptions) {
 									if (info.description.search(new RegExp(pattern,"s")) >= 0) {
 										//skip this one
@@ -1953,16 +1951,36 @@ export class NodeBot extends PerforceStatefulBot implements NodeBotInterface {
 										break
 									}
 								}
+								if (descriptionIsExcluded) { //TODO: consider disregard flag
+									continue
+								}
 							}
-
-							if (!descriptionIsExcluded) { //TODO: consider disregard flag
-								mergeActions.push(action)
+							
+							if (edgeProperties.changelistFilter) {
+								const changelistFilter = changelistFilters.get(edgeProperties.changelistFilter.name)
+								if (changelistFilter) {
+									const filterResult = await changelistFilter.shouldFilter(info, action, edgeProperties, this.nodeBotLogger)
+									switch(filterResult.result) {
+										case 'allow':
+											action.flags.add('#filter_approved#')
+											break
+										case 'block':
+											action.flags.add('#filter_blocked#')
+											break
+										case 'skip':
+											continue
+									}
+									if (filterResult.message) {
+										action.flags.add(`#filter_reason_${filterResult.message}`)
+									}
+								}
 							}
 						}
+						mergeActions.push(action)
+					}
 
-						if (mergeActions.length > 0) {
-							info.targets = mergeActions
-						}
+					if (mergeActions.length > 0) {
+						info.targets = mergeActions
 					}
 
 					if (info.targets && flags.has('ignore')) {
diff --git a/Engine/Extras/RoboMerge/v3/src/robo/status-types.ts b/Engine/Extras/RoboMerge/v3/src/robo/status-types.ts
--- a/Engine/Extras/RoboMerge/v3/src/robo/status-types.ts
+++ b/Engine/Extras/RoboMerge/v3/src/robo/status-types.ts
@@ -25,6 +25,7 @@ export type FailureKind =
 	'Syntax error' |
 	'Disallowed files' |
 	'Approval required' |
+	'Blocked' |
 	'Unit Test error'
 
 export enum AvailableTypeEnum {
diff --git a/Engine/Extras/RoboMerge/v3/src/robo/targets.ts b/Engine/Extras/RoboMerge/v3/src/robo/targets.ts
--- a/Engine/Extras/RoboMerge/v3/src/robo/targets.ts
+++ b/Engine/Extras/RoboMerge/v3/src/robo/targets.ts
@@ -1,7 +1,8 @@
 // Copyright Epic Games, Inc. All Rights Reserved.
 
 import { Branch, MergeAction, MergeMode } from './branch-interfaces';
-import { BranchGraphInterface, PendingChange, Target, TargetInfo } from './branch-interfaces';
+import { BranchGraphInterface, ChangeInfo, PendingChange, Target, TargetInfo } from './branch-interfaces';
+import { EdgeOptions, EdgeProperties, NodeOptions } from './branchdefs'
 import { BotName, Edge, GraphAPI, Node, makeTargetName } from './graph';
 import { ContextualLogger } from '../common/logger';
 import { Change } from '../common/perforce';
@@ -364,6 +365,23 @@ function isFlag(arg: string) {
 	return ALLOWED_RAW_FLAGS.indexOf(arg.toLowerCase()) >= 0 || '$#'.indexOf(arg[0]) >= 0
 }
 
+export type ChangelistFilterResult = {
+	result: 'allow' | 'skip' | 'block',
+	message?: string
+}
+export class ChangelistFilter {
+
+	static async shouldFilter(_info: ChangeInfo, _action: MergeAction, _edgeProperties: EdgeOptions, _logger: ContextualLogger) : Promise<ChangelistFilterResult> {
+		return {result: 'block'}
+	}
+
+	// Validation function returns a tuple of an array each for errors and warnings
+	static async validate(_branches: NodeOptions[], _edge: EdgeProperties, _isPreviewing: boolean, _logger: ContextualLogger) : Promise<[string[],string[]]> {
+		return [[],[]]
+	}
+}
+
+export const changelistFilters = new Map<string, typeof ChangelistFilter>()
 
 /**
 
