Start to harden args access

diff --git a/Engine/Extras/RoboMerge/v3/src/robo/autobranchupdater.ts b/Engine/Extras/RoboMerge/v3/src/robo/autobranchupdater.ts
--- a/Engine/Extras/RoboMerge/v3/src/robo/autobranchupdater.ts
+++ b/Engine/Extras/RoboMerge/v3/src/robo/autobranchupdater.ts
@@ -177,7 +177,7 @@ export class AutoBranchUpdater implements Bot {
 			const validationErrorText = result.validationErrors.join('\n')
 			const file: SlackFile = {
 				content: validationErrorText,
-				channels: RoboArgs.args.slackAlertChannel,
+				channels: RoboArgs.slackAlertChannel,
 				filename: "branchspecerrors.txt"
 			}
 			postToRobomergeAlerts(errText, [file])
diff --git a/Engine/Extras/RoboMerge/v3/src/robo/graphbot.ts b/Engine/Extras/RoboMerge/v3/src/robo/graphbot.ts
--- a/Engine/Extras/RoboMerge/v3/src/robo/graphbot.ts
+++ b/Engine/Extras/RoboMerge/v3/src/robo/graphbot.ts
@@ -71,7 +71,7 @@ export class GraphBot implements BotEventHandler {
 			const parseError = result.validationErrors.length === 0 ? 'Failed to parse' : result.validationErrors.join('\n')
 			const file: SlackFile = {
 				content: parseError,
-				channels: RoboArgs.args.slackAlertChannel,
+				channels: RoboArgs.slackAlertChannel,
 				filename: "branchspecerrors.txt"
 			}
 			await postToRobomergeAlerts(`@here Failed to create bot ${botname.toUpperCase()}.\nHalting service startup.`, [file])
diff --git a/Engine/Extras/RoboMerge/v3/src/robo/ipc.ts b/Engine/Extras/RoboMerge/v3/src/robo/ipc.ts
--- a/Engine/Extras/RoboMerge/v3/src/robo/ipc.ts
+++ b/Engine/Extras/RoboMerge/v3/src/robo/ipc.ts
@@ -374,7 +374,7 @@ export class IPC {
 	}
 
 	private logError(msg: string) {
-		if (RoboArgs.args.runningFunctionalTests) {
+		if (RoboArgs.isRunningFunctionalTests()) {
 			this.ipcLogger.info(msg)
 		} else {
 			this.ipcLogger.error(msg)
diff --git a/Engine/Extras/RoboMerge/v3/src/robo/nodebot.ts b/Engine/Extras/RoboMerge/v3/src/robo/nodebot.ts
--- a/Engine/Extras/RoboMerge/v3/src/robo/nodebot.ts
+++ b/Engine/Extras/RoboMerge/v3/src/robo/nodebot.ts
@@ -1025,8 +1025,8 @@ export class NodeBot extends PerforceStatefulBot implements NodeBotInterface {
 
 		// Verify that we have the client/user information for all locked files
 		let lockedFiles: ExclusiveFileDetails[] = blockageInfo.additionalInfo.exclusiveFiles
-		if (lockedFiles.length > RoboArgs.getExclusiveLockOpenedsToRun()) {
-			lockedFiles = [...lockedFiles.slice(0,RoboArgs.getExclusiveLockOpenedsToRun()), ...(await this.p4.getOpenedDetails(lockedFiles.slice(RoboArgs.getExclusiveLockOpenedsToRun())))]
+		if (lockedFiles.length > RoboArgs.exclusiveLockOpenedsToRun) {
+			lockedFiles = [...lockedFiles.slice(0,RoboArgs.exclusiveLockOpenedsToRun), ...(await this.p4.getOpenedDetails(lockedFiles.slice(RoboArgs.exclusiveLockOpenedsToRun)))]
 		}
 
 		const unknownClients = lockedFiles.filter((file) => file.client.length == 0)
diff --git a/Engine/Extras/RoboMerge/v3/src/robo/notifications.ts b/Engine/Extras/RoboMerge/v3/src/robo/notifications.ts
--- a/Engine/Extras/RoboMerge/v3/src/robo/notifications.ts
+++ b/Engine/Extras/RoboMerge/v3/src/robo/notifications.ts
@@ -56,7 +56,7 @@ export function isUserAKnownBot(user: string) {
 
 export async function postToRobomergeAlerts(message: string, replies?: (string|SlackFile|SlackMessage)[]) {
 	if (!args.devMode) {
-		return postMessageToChannel(`${message}\n\n${RoboArgs.args.externalUrl}`, args.slackAlertChannel, SlackMessageStyles.DANGER, replies)
+		return postMessageToChannel(`${message}\n\n${RoboArgs.externalUrl}`, args.slackAlertChannel, SlackMessageStyles.DANGER, replies)
 	}
 }
 
diff --git a/Engine/Extras/RoboMerge/v3/src/robo/persistedmessages.ts b/Engine/Extras/RoboMerge/v3/src/robo/persistedmessages.ts
--- a/Engine/Extras/RoboMerge/v3/src/robo/persistedmessages.ts
+++ b/Engine/Extras/RoboMerge/v3/src/robo/persistedmessages.ts
@@ -87,7 +87,7 @@ export class PersistedSlackMessage implements PersistedMessagePayload {
 export abstract class PersistedSlackMessages {
 
 	static get(botSettings: Settings, logger: ContextualLogger): PersistedSlackMessages {
-		if (RoboArgs.useMongo()) {
+		if (RoboArgs.useMongo) {
 			return new DBPersistedSlackMessages(botSettings.botname)
 		}
 		return new FilePersistedSlackMessages(botSettings, logger)
diff --git a/Engine/Extras/RoboMerge/v3/src/robo/roboargs.ts b/Engine/Extras/RoboMerge/v3/src/robo/roboargs.ts
--- a/Engine/Extras/RoboMerge/v3/src/robo/roboargs.ts
+++ b/Engine/Extras/RoboMerge/v3/src/robo/roboargs.ts
@@ -10,19 +10,27 @@ export class RoboArgs {
 		this._args = args
 	}
 
-	static getExclusiveLockOpenedsToRun() {
+	static get exclusiveLockOpenedsToRun() {
 		return this._args.exclusiveLockOpenedsToRun
 	}
 
-	static getMaxStompFileCount() {
+	static get externalUrl() {
+		return this._args.externalUrl
+	}
+
+	static get isRunningFunctionalTests() {
+		return this._args.runningFunctionalTests
+	}
+
+	static get maxStompFileCount() {
 		return this._args.maxStompFileCount
 	}
 
-	static useMongo() {
-		return this._args.mongoDB_URI.length > 0
+	static get slackAlertChannel() {
+		return this._args.slackAlertChannel
 	}
 
-	static get args() {
-		return this._args
+	static get useMongo() {
+		return this._args.mongoDB_URI.length > 0
 	}
 }
\ No newline at end of file
