Add support to continue writing out the json persistence files while using the DB. Be aware that SlackMessage data will not be written out, so if the backups are required, all slack message data will have been lost.
Fix up a bunch of misuses of RoboArgs from last CL
Bump package version

diff --git a/Engine/Extras/RoboMerge/v3/package.json b/Engine/Extras/RoboMerge/v3/package.json
--- a/Engine/Extras/RoboMerge/v3/package.json
+++ b/Engine/Extras/RoboMerge/v3/package.json
@@ -1,6 +1,6 @@
 {
   "name": "robomerge",
-  "version": "3.4.0",
+  "version": "3.4.1",
   "description": "Automatic P4 merge management bot.",
   "repository": {},
   "keywords": [
diff --git a/Engine/Extras/RoboMerge/v3/src/robo/edgebot.ts b/Engine/Extras/RoboMerge/v3/src/robo/edgebot.ts
--- a/Engine/Extras/RoboMerge/v3/src/robo/edgebot.ts
+++ b/Engine/Extras/RoboMerge/v3/src/robo/edgebot.ts
@@ -264,7 +264,7 @@ class EdgeBotImpl extends PerforceStatefulBot {
 		for (const err of errors) {
 			const match = err.match(EXCLUSIVE_CHECKOUT_REGEX)
 			if (match) {
-				if (paths.length < RoboArgs.getExclusiveLockOpenedsToRun()) {
+				if (paths.length < RoboArgs.exclusiveLockOpenedsToRun) {
 					paths.push({depotPath: match[1] + match[2], name: match[2]})
 				} else {
 					results.push({depotPath: match[1] + match[2], name: match[2], user: "", client: ""})
@@ -819,7 +819,7 @@ class EdgeBotImpl extends PerforceStatefulBot {
 					client: pending.change.targetWorkspace
 				}
 			}
-			if (failure.additionalInfo.describeResult.entries.length > RoboArgs.getMaxStompFileCount()) {
+			if (failure.additionalInfo.describeResult.entries.length > RoboArgs.maxStompFileCount) {
 				failure.additionalInfo.describeResult = undefined
 				failure.additionalInfo.disableStompReason = "Changelist too large"
 			}
diff --git a/Engine/Extras/RoboMerge/v3/src/robo/robo.ts b/Engine/Extras/RoboMerge/v3/src/robo/robo.ts
--- a/Engine/Extras/RoboMerge/v3/src/robo/robo.ts
+++ b/Engine/Extras/RoboMerge/v3/src/robo/robo.ts
@@ -229,6 +229,12 @@ const COMMAND_LINE_ARGS: {[param: string]: Arg<any>} = {
 		env: 'ROBO_MONGODB_URI',
 		dflt: ''
 	},
+	alwaysPersistToFile: {
+		match: /^(-alwaysPersistToFile)$/,
+		parse: str => str === "false" ? false : true,
+		env: 'ROBO_ALWAYS_PERSIST_TO_FILE',
+		dflt: false
+	},
 	persistenceDir: {
 		match: /^-persistenceDir=(.+)$/,
 		env: 'ROBO_PERSISTENCE_DIR',
@@ -277,7 +283,7 @@ settingsInit(args)
 /**
  * Connect to MongoDB.
  */
-if (RoboArgs.useMongo()) {
+if (RoboArgs.useMongo) {
 	var mongoose = require('mongoose')
 	mongoose.connect(args.mongoDB_URI)
 	mongoose.connection.on('connected', () => {
diff --git a/Engine/Extras/RoboMerge/v3/src/robo/roboargs.ts b/Engine/Extras/RoboMerge/v3/src/robo/roboargs.ts
--- a/Engine/Extras/RoboMerge/v3/src/robo/roboargs.ts
+++ b/Engine/Extras/RoboMerge/v3/src/robo/roboargs.ts
@@ -33,4 +33,8 @@ export class RoboArgs {
 	static get useMongo() {
 		return this._args.mongoDB_URI.length > 0
 	}
+
+	static get persistSettingsToFile() {
+		return !this.useMongo || this._args.alwaysPersistToFile
+	}
 }
\ No newline at end of file
diff --git a/Engine/Extras/RoboMerge/v3/src/robo/settings.ts b/Engine/Extras/RoboMerge/v3/src/robo/settings.ts
--- a/Engine/Extras/RoboMerge/v3/src/robo/settings.ts
+++ b/Engine/Extras/RoboMerge/v3/src/robo/settings.ts
@@ -138,7 +138,7 @@ export class Settings {
 
 		const settings = new Settings(botname, parentLogger, p4)
 
-		if (RoboArgs.useMongo()) {
+		if (RoboArgs.useMongo) {
 			// first try to get the bot data from the DB
 			const botRecords = await BotModel.find({botname: botname})
 			settings.settingsLogger.info(`Loaded ${botRecords.length} node records`)
@@ -216,7 +216,7 @@ export class Settings {
 		const previousObjVersion = new semver.SemVer(objVersion)
 
 		// Transition slack messages out of the settings object and into the database
-		if (loadedFromFile && RoboArgs.useMongo()) {
+		if (loadedFromFile && RoboArgs.useMongo) {
 			settings.settingsLogger.info("Migrated settings from json to MongoDB")
 			await migrateMessagesToDB(settings)
 		}
@@ -246,7 +246,7 @@ export class Settings {
 
 	_saveObject(opts?: SaveObjectOpts) {
 		if (this.enableSave) {
-			if (RoboArgs.useMongo()) {
+			if (RoboArgs.useMongo) {
 				if (opts?.name) {
 					const saveToDB = () => {
 						this.savingState.set(opts.name!, 'Saving')
@@ -290,7 +290,7 @@ export class Settings {
 					})
 				}
 			}
-			else {
+			if (RoboArgs.persistSettingsToFile) {
 				const persistentFile = this.getFilename('PERSISTENCE')
 				const tempFile = this.getFilename('TEMP')
 
@@ -348,7 +348,7 @@ export class Settings {
 			const additionalInfo: MergeConflictAdditionalInformation = edge.pause?.blockage?.additionalInfo
 			if (isA_MergeConflictAdditionalInformation(additionalInfo)) {
 				if (additionalInfo?.describeResult) {
-					if (!trim || additionalInfo.describeResult.entries.length > RoboArgs.getMaxStompFileCount()) {
+					if (!trim || additionalInfo.describeResult.entries.length > RoboArgs.maxStompFileCount) {
 						additionalInfo.disableStompReason = "Changelist too large"
 						additionalInfo.describeResult = undefined
 					}
