Strip non-UE jiras if they are listed after a non-UE jira and include UE jiras if they are listed after a non-UE jira

diff --git a/Engine/Extras/RoboMerge/v3/src/robo/edgebot.ts b/Engine/Extras/RoboMerge/v3/src/robo/edgebot.ts
--- a/Engine/Extras/RoboMerge/v3/src/robo/edgebot.ts
+++ b/Engine/Extras/RoboMerge/v3/src/robo/edgebot.ts
@@ -26,6 +26,7 @@ import { getIntegrationOwner } from "./targets";
 
 const FAILED_CHANGELIST_PAUSE_TIMEOUT_SECONDS = 15 * 60
 const JIRA_REGEX = /^\s*#jira\s+(.*)/i
+const UE_JIRA_REGEX = /(?:^|[^a-z0-9])((?:ue)-\d+)(?=$|[^a-z0-9])/gi
 
 // type ResolveResultDetail = 'quick' | 'detailed'
 
@@ -325,13 +326,17 @@ class EdgeBotImpl extends PerforceStatefulBot {
 
 				// strip non UE- Jira tags
 				const jiraTagMatch = line.match(JIRA_REGEX)
-				if (jiraTagMatch && !jiraTagMatch[1].toLowerCase().startsWith('ue-')) {
-					continue
-				}
 
-				// further sanitisation can go here
+				if (jiraTagMatch) {
+					const ueJiras = Array.from(jiraTagMatch[1].matchAll(UE_JIRA_REGEX)).map(m => m[1])
 
-				description += line + '\n'
+					if (ueJiras.length > 0) {
+						description += `#jira ${ueJiras.join(" ")}\n`
+					}
+				}
+				else {
+					description += line + "\n"
+				}
 			}
 		}
 		else {
