Handle unloaded workspaces for gate, branchspec, and persistence workspaces

diff --git a/Engine/Extras/RoboMerge/v3/src/robo/gate.ts b/Engine/Extras/RoboMerge/v3/src/robo/gate.ts
--- a/Engine/Extras/RoboMerge/v3/src/robo/gate.ts
+++ b/Engine/Extras/RoboMerge/v3/src/robo/gate.ts
@@ -473,13 +473,16 @@ export class Gate {
 	
 		const p4 = this.context.p4!
 	
-		const workspaceQueryResult = await p4.find_workspace_by_name(workspace.name)
+		const workspaceQueryResult = await p4.find_workspace_by_name(workspace.name, {includeUnloaded: true})
 		if (workspaceQueryResult.length === 0) {
 			await p4.newWorkspace(workspace.name, {
 				Stream: stream,
 				Root: workspace.directory
 			})
 		}
+		else if (workspaceQueryResult[0].IsUnloaded) {
+			await p4.reloadWorkspace(workspace.name)
+		}
 	
 		const syncResult = await p4.sync(workspace, lastGoodCLPath, {opts:[P4_FORCE]})
 		if (!syncResult || syncResult.length != 1) {
diff --git a/Engine/Extras/RoboMerge/v3/src/robo/robo.ts b/Engine/Extras/RoboMerge/v3/src/robo/robo.ts
--- a/Engine/Extras/RoboMerge/v3/src/robo/robo.ts
+++ b/Engine/Extras/RoboMerge/v3/src/robo/robo.ts
@@ -443,7 +443,7 @@ async function init(logger: ContextualLogger) {
 		}
 
 		// Ensure we have a workspace for branch specs
-		const workspace: Object[] = await p4.find_workspace_by_name(args.branchSpecsWorkspace)
+		const workspace = await p4.find_workspace_by_name(args.branchSpecsWorkspace, {includeUnloaded: true})
 		if (workspace.length === 0) {
 			logger.info(`Cannot find workspace ${args.branchSpecsWorkspace}, creating a new one.`)
 
@@ -451,6 +451,9 @@ async function init(logger: ContextualLogger) {
 				await p4.newBranchSpecWorkspace(autoUpdaterConfig.workspace, args.branchSpecsRootPath)
 			}
 		}
+		else if (workspace[0].IsUnloaded) {
+			await p4.reloadWorkspace(args.branchSpecsWorkspace)
+		}
 
 		// make sure we've got the latest branch specs
 		logger.info('Syncing latest branch specs')
@@ -467,7 +470,7 @@ async function init(logger: ContextualLogger) {
 	}
 
 	if (args.persistenceBackupFrequency > 0 && !args.previewOnly) {
-		const workspace: Object[] = await p4.find_workspace_by_name(args.persistenceBackupWorkspace)
+		const workspace = await p4.find_workspace_by_name(args.persistenceBackupWorkspace, {includeUnloaded: true})
 		if (workspace.length === 0) {
 			logger.info(`Cannot find workspace ${args.persistenceBackupWorkspace}, creating a new one.`)
 
@@ -485,6 +488,9 @@ async function init(logger: ContextualLogger) {
 				await p4.newWorkspace(args.persistenceBackupWorkspace, params)
 			}
 		}
+		else if (workspace[0].IsUnloaded) {
+			await p4.reloadWorkspace(args.persistenceBackupWorkspace)
+		}
 	} else {
 		logger.warn('Persistence backup not configured!')
 	}
