Add a cap on the number of files that can be reconsidered

diff --git a/Engine/Extras/RoboMerge/v3/src/robo/nodebot.ts b/Engine/Extras/RoboMerge/v3/src/robo/nodebot.ts
--- a/Engine/Extras/RoboMerge/v3/src/robo/nodebot.ts
+++ b/Engine/Extras/RoboMerge/v3/src/robo/nodebot.ts
@@ -384,8 +384,8 @@ export class NodeBot extends PerforceStatefulBot implements NodeBotInterface {
 
 	async processQueuedChange(fromQueue: QueuedChange) {
 		if (this.previewMode) {
-			this._log_action(`(PREVIEW MODE) - Skipping manually queued change ${fromQueue.cl} on ${this.fullName}, requested by ${fromQueue.who}`)
-			return
+			//this._log_action(`(PREVIEW MODE) - Skipping manually queued change ${fromQueue.cl} on ${this.fullName}, requested by ${fromQueue.who}`)
+			//return
 		}
 		let logMessage = `Processing manually queued change ${fromQueue.cl} on ${this.fullName}, requested by ${fromQueue.who}`
 		if (fromQueue.workspace) {
@@ -400,52 +400,23 @@ export class NodeBot extends PerforceStatefulBot implements NodeBotInterface {
 			specifiedTargetBranch = this._getBranch(fromQueue.targetBranchName)
 
 			if (!specifiedTargetBranch) {
-				if (this.slackMessages) {
-					const dm: SlackMessage = {
-						text: `Manually queued change ${fromQueue.cl} on ${this.fullName}, requested by ${fromQueue.who}, ` + 
-								`specifies non-existant branch ${fromQueue.targetBranchName}. Aborting reconsider.`,
-						style: SlackMessageStyles.DANGER,
-						channel: '',
-						mrkdwn: true
-					}
-					
-					let emailAddress = await this.findEmail(fromQueue.who)
-					this.slackMessages.postDM(emailAddress, fromQueue.cl, this.branch, dm)
-				}
-				else {
-					// We somehow have a specified branch from a reconsider request, but it doesn't exist in the branch map. Report the error
-					this.sendErrorEmail(
-						new Recipients(fromQueue.who),
-						`Invalid Reconsider Request -- Branch "${fromQueue.targetBranchName}" does not exist`,
-						`Manually queued change ${fromQueue.cl} on ${this.fullName}, requested by ${fromQueue.who}, ` + 
-							`specifies non-existant branch ${fromQueue.targetBranchName}. Aborting reconsider.`
-					)
-				}
+				await this.sendErrorNotificationToUser(
+					fromQueue.who,
+					fromQueue.cl,
+					`Invalid Reconsider Request -- Branch "${fromQueue.targetBranchName}" does not exist`,
+					`Manually queued change ${fromQueue.cl} on ${this.fullName}, requested by ${fromQueue.who}, ` +
+						`specifies non-existant branch ${fromQueue.targetBranchName}. Aborting reconsider.`
+				)
 				return
 			}
 			else if (specifiedTargetBranch === this.branch) {
-				if (this.slackMessages) {
-					const dm: SlackMessage = {
-						text: `Manually queued change ${fromQueue.cl} on ${this.fullName}, requested by ${fromQueue.who}, ` +
-								`specifies integration from branch ${fromQueue.targetBranchName} into itself. Aborting reconsider.`,
-						style: SlackMessageStyles.DANGER,
-						channel: '',
-						mrkdwn: true
-					}
-					
-					let emailAddress = await this.findEmail(fromQueue.who)
-					this.slackMessages.postDM(emailAddress, fromQueue.cl, this.branch, dm)
-				}
-				else {
-					// We somehow have targeted our own branch with a merge. Report the issue
-					this.sendErrorEmail(
-						new Recipients(fromQueue.who),
-						`Invalid Reconsider Request -- Trying to merge "${fromQueue.targetBranchName}" into itself`,
-						`Manually queued change ${fromQueue.cl} on ${this.fullName}, requested by ${fromQueue.who}, ` +
-							`specifies integration from branch ${fromQueue.targetBranchName} into itself. Aborting reconsider.`
-					)
-				}
-				// Persist the remaining queued changes and carry on
+				await this.sendErrorNotificationToUser(
+					fromQueue.who,
+					fromQueue.cl,
+					`Invalid Reconsider Request -- Trying to merge "${fromQueue.targetBranchName}" into itself`,
+					`Manually queued change ${fromQueue.cl} on ${this.fullName}, requested by ${fromQueue.who}, ` +
+						`specifies integration from branch ${fromQueue.targetBranchName} into itself. Aborting reconsider.`
+				)
 				return
 			}
 
@@ -458,29 +429,20 @@ export class NodeBot extends PerforceStatefulBot implements NodeBotInterface {
 		const changeResult = await this._getChange(fromQueue.cl)
 
 		if (typeof changeResult === 'string') {
-
-			if (this.slackMessages) {
-
-				let reconsiderErrorMessage = `While processing manually queued change ${fromQueue.cl} on ${this.fullName}, `
-				reconsiderErrorMessage += specifiedTargetBranch ? `targeted specifically at ${specifiedTargetBranch.name}` : ''
-				reconsiderErrorMessage += `\n\nError while querying P4 for change ${fromQueue.cl}: ${changeResult}`
-
-				const dm: SlackMessage = {
-					text: reconsiderErrorMessage,
-					style: SlackMessageStyles.DANGER,
-					channel: '',
-					mrkdwn: true
-				}
-				
-				let emailAddress = await this.findEmail(fromQueue.who)
-				this.slackMessages.postDM(emailAddress, fromQueue.cl, this.branch, dm)
-			}
+			let reconsiderErrorMessage = `While processing manually queued change ${fromQueue.cl} on ${this.fullName}, `
+			reconsiderErrorMessage += specifiedTargetBranch ? `targeted specifically at ${specifiedTargetBranch.name}` : ''
+			reconsiderErrorMessage += `\n\nError while querying P4 for change ${fromQueue.cl}: ${changeResult}`
+
+			await this.sendErrorNotificationToUser(
+				fromQueue.who,
+				fromQueue.cl,
+				`Reconsider Request Failed -- Error Querying Perforce`,
+				reconsiderErrorMessage
+			)
 
 			// change is invalid
-
 			// @todo distinguish between Perforce errors and change not existing in branch
 
-
 			this.nodeBotLogger.error(`${this.fullName} - Error while querying P4 for change ${fromQueue.cl}: ${changeResult}`)
 			// @todo email admin page user when users have to log in
 			// any way to show error without blocking page?
@@ -491,6 +453,45 @@ export class NodeBot extends PerforceStatefulBot implements NodeBotInterface {
 		const change = changeResult as Change
 		change.isUserRequest = true
 
+		// Check if the change exceeds the maximum file limit for reconsider requests
+		if (RoboArgs.maxReconsiderFiles > 0) {
+			try {
+				const sizeResult = await this.p4.sizes(`${this.branch.rootPath}@=${change.change}`, true)
+				const fileCount = sizeResult[0].fileCount
+
+				if (fileCount > RoboArgs.maxReconsiderFiles) {
+					const limitExceededMessage = `Reconsider request for CL ${fromQueue.cl} on ${this.fullName} exceeds the maximum file limit.\n\n` +
+						`This change contains ${fileCount} file(s), but the maximum allowed for reconsider is ${RoboArgs.maxReconsiderFiles}.\n\n` +
+						`Please perform this merge manually instead of using the reconsider feature.`
+
+					await this.sendErrorNotificationToUser(
+						fromQueue.who,
+						fromQueue.cl,
+						`Reconsider Request Rejected -- Too Many Files`,
+						limitExceededMessage
+					)
+
+					this.nodeBotLogger.warn(`${this.fullName} - Reconsider request for CL ${fromQueue.cl} rejected: ${fileCount} files exceeds limit of ${RoboArgs.maxReconsiderFiles}`)
+					return
+				}
+			}
+			catch (err) {
+				const errorMessage = `Reconsider request for CL ${fromQueue.cl} on ${this.fullName} could not be processed.\n\n` +
+					`Failed to check file count: ${err}\n\n` +
+					`Please contact your Robomerge administrator or perform this merge manually.`
+
+				await this.sendErrorNotificationToUser(
+					fromQueue.who,
+					fromQueue.cl,
+					`Reconsider Request Failed -- Unable to Verify File Count`,
+					errorMessage
+				)
+
+				this.nodeBotLogger.error(`${this.fullName} - Failed to check file count for CL ${fromQueue.cl}: ${err}`)
+				return
+			}
+		}
+
 		const additionalFlags = fromQueue.additionalFlags || []
 		// By default, we'll submit any changes as a result of processing the CL. In some cases, such as
 		// the create a shelf feature, we actually want to hold off on submitting any changes.
@@ -1494,6 +1495,27 @@ export class NodeBot extends PerforceStatefulBot implements NodeBotInterface {
 		}
 	}
 
+	private async sendErrorNotificationToUser(who: string, cl: number, emailSubject: string, messageText: string) {
+		if (this.slackMessages) {
+			const dm: SlackMessage = {
+				text: messageText,
+				style: SlackMessageStyles.DANGER,
+				channel: '',
+				mrkdwn: true
+			}
+
+			let emailAddress = await this.findEmail(who)
+			this.slackMessages.postDM(emailAddress, cl, this.branch, dm)
+		}
+		else {
+			this.sendErrorEmail(
+				new Recipients(who),
+				emailSubject,
+				messageText
+			)
+		}
+	}
+
 	private async getChangeAuthor(change: Change) {
 		const descriptionParser = this.branch.changelistParser && descriptionParsers.get(this.branch.changelistParser)
 		if (descriptionParser) {
diff --git a/Engine/Extras/RoboMerge/v3/src/robo/robo.ts b/Engine/Extras/RoboMerge/v3/src/robo/robo.ts
--- a/Engine/Extras/RoboMerge/v3/src/robo/robo.ts
+++ b/Engine/Extras/RoboMerge/v3/src/robo/robo.ts
@@ -171,6 +171,15 @@ const COMMAND_LINE_ARGS: {[param: string]: Arg<any>} = {
 		dflt: 30000
 	},
 
+	maxReconsiderFiles: {
+		match: /^-maxReconsiderFiles=([0-9]+)$/,
+		env: 'ROBO_MAX_RECONSIDER_FILES',
+		parse: (str: string) => {
+			return parseInt(str)
+		},
+		dflt: 50000
+	},
+
 	runningFunctionalTests: {
 		match: /^(-runningFunctionalTests)$/,
 		parse: str => str === "false" ? false : true,
diff --git a/Engine/Extras/RoboMerge/v3/src/robo/roboargs.ts b/Engine/Extras/RoboMerge/v3/src/robo/roboargs.ts
--- a/Engine/Extras/RoboMerge/v3/src/robo/roboargs.ts
+++ b/Engine/Extras/RoboMerge/v3/src/robo/roboargs.ts
@@ -26,6 +26,10 @@ export class RoboArgs {
 		return this._args.maxStompFileCount
 	}
 
+	static get maxReconsiderFiles() {
+		return this._args.maxReconsiderFiles
+	}
+
 	static get slackAlertChannel() {
 		return this._args.slackAlertChannel
 	}
