Fix a failure to save due to object size resulting in a state where the node never again saves its state.

diff --git a/Engine/Extras/RoboMerge/v3/src/robo/settings.ts b/Engine/Extras/RoboMerge/v3/src/robo/settings.ts
--- a/Engine/Extras/RoboMerge/v3/src/robo/settings.ts
+++ b/Engine/Extras/RoboMerge/v3/src/robo/settings.ts
@@ -261,17 +261,24 @@ export class Settings {
 							}
 						})
 						.catch(() => {
-							const size = Buffer.byteLength(JSON.stringify(this.object[opts.name!]));
-							this.settingsLogger.warn(`Unable to save ${this.botname}:${opts.name}. Size: ${size}. Attempting to shrink and resave.`)
-							this._shrinkNode(opts.name!)
-							const shrunkSize = Buffer.byteLength(JSON.stringify(this.object[opts.name!]));
-							if (shrunkSize < size) {
+							if (this.savingState.get(opts.name!) === 'NeedsSave') {
+								this.settingsLogger.warn(`Unable to save ${this.botname}:${opts.name}. Attempting resave before applying mitigation as settings have changed.`)
 								saveToDB()
-							}
+							} 
 							else {
-								const errorMsg = `Unable to save ${this.botname}:${opts.name}. Size: ${size}.`
-								this.settingsLogger.error(errorMsg)
-								postToRobomergeAlerts(errorMsg)
+								const size = Buffer.byteLength(JSON.stringify(this.object[opts.name!]));
+								this._shrinkNode(opts.name!)
+								const shrunkSize = Buffer.byteLength(JSON.stringify(this.object[opts.name!]));
+								if (shrunkSize < size) {
+									this.settingsLogger.warn(`Unable to save ${this.botname}:${opts.name}. Size: ${size}. Attempting resave with shrunk size ${shrunkSize}.`)
+									saveToDB()
+								}
+								else {
+									const errorMsg = `Unable to save ${this.botname}:${opts.name}. Size: ${size}.`
+									this.settingsLogger.error(errorMsg)
+									postToRobomergeAlerts(errorMsg)
+									this.savingState.delete(opts.name!)
+								}
 							}
 						})
 					}
