Detect multi-server environment in a more nuanced fashion

diff --git a/Engine/Extras/RoboMerge/v3/src/common/perforce.ts b/Engine/Extras/RoboMerge/v3/src/common/perforce.ts
--- a/Engine/Extras/RoboMerge/v3/src/common/perforce.ts
+++ b/Engine/Extras/RoboMerge/v3/src/common/perforce.ts
@@ -556,7 +556,7 @@ export class PerforceContext {
 				serverID,
 				serverAddress,
 				serverVersion,
-				multiServerEnvironment: serversOutput.length > 1,
+				multiServerEnvironment: serversOutput.some(server => server['Services'] == 'edge-server'),
 				swarmURL: swarmOutput.length > 0 ? swarmOutput.trim() : undefined
 			}
 			this.servers.set(serverID, serverConfig)
