Internalize getPingTarget

diff --git a/Engine/Extras/RoboMerge/v3/src/robo/notifications.ts b/Engine/Extras/RoboMerge/v3/src/robo/notifications.ts
--- a/Engine/Extras/RoboMerge/v3/src/robo/notifications.ts
+++ b/Engine/Extras/RoboMerge/v3/src/robo/notifications.ts
@@ -92,17 +92,6 @@ export function isUserSlackGroup(user: string) {
 	return user.startsWith('@') || user.startsWith('<')
 }
 
-export async function getPingTarget(userEmail: string|null, owner: string, slackMessages: SlackMessages) {
-	const slackUser = userEmail ? await slackMessages.getSlackUser(userEmail) : null
-	return slackUser 
-	? `<@${slackUser}>` 
-	: isUserSlackGroup(owner)
-		? owner 
-		: isUserAKnownBot(owner)
-			? ""
-			: `@${owner}`
-}
-
 function formatDuration(durationSeconds: number) {
 	const underSixHours = durationSeconds < 6 * 3600
 	const durationBits: string[] = []
@@ -418,9 +407,20 @@ export class SlackMessages {
 		}
 	}
 
+	async getPingTarget(userEmailAddress: Promise<string|null>, owner: string) {
+		const userEmail = await userEmailAddress
+		const slackUser = userEmail ? await this.getSlackUser(userEmail) : null
+		return slackUser 
+		? `<@${slackUser}>` 
+		: isUserSlackGroup(owner)
+			? owner 
+			: isUserAKnownBot(owner)
+				? ""
+				: `@${owner}`
+	}
+
 	// With their email address, we can get their user ID via Slack API
 	getSlackUser(emailAddress: string) {
-		
 		return this.slack.lookupUserIdByEmail(emailAddress.toLowerCase())
 	}
 }
@@ -537,8 +537,7 @@ export class BotNotifications implements BotEventHandler {
 		}
 		const issue = blockage.failure.kind.toLowerCase()
 
-		const userEmail = await blockage.ownerEmail
-		const channelPing = await getPingTarget(userEmail, blockage.owner, this.slackMessages)
+		const channelPing = await this.slackMessages.getPingTarget(blockage.ownerEmail, blockage.owner)
 
 		const isBotUser = isUserAKnownBot(blockage.owner)
 		let text = ""
@@ -576,6 +575,7 @@ export class BotNotifications implements BotEventHandler {
 		const channelsToPostTo = this.getChannelsToPostTo(sourceBranch, targetBranch?.upperName)
 		let usersToInvite: Set<string> = new Set()
 
+		const userEmail = await blockage.ownerEmail
 		if (userEmail) {
 			usersToInvite.add(userEmail)
 		}
@@ -725,7 +725,7 @@ export class BotNotifications implements BotEventHandler {
 			}
 
 			const triager = await this.getTriager(evt.triager, evt.triagerEmail)
-			const channelPing = await getPingTarget(await evt.ownerEmail, conflict.owner, this.slackMessages)
+			const channelPing = await this.slackMessages.getPingTarget(evt.ownerEmail, conflict.owner)
 
 			let messageTargets = []
 			let loggerTargets = []
