Make changes take an options block

diff --git a/Engine/Extras/RoboMerge/v3/src/common/perforce.ts b/Engine/Extras/RoboMerge/v3/src/common/perforce.ts
--- a/Engine/Extras/RoboMerge/v3/src/common/perforce.ts
+++ b/Engine/Extras/RoboMerge/v3/src/common/perforce.ts
@@ -289,6 +289,14 @@ interface ExecZtagOpts extends ExecOpts {
 	reduce?: boolean; // collapse the multiple entries to a single object, useful for problem parses that end up in multiple entries despite being a single result
 }
 
+export interface ChangesOpts {
+	since?: number | string
+	until?: number | string
+	limit?: number
+	status?: ChangelistStatus
+	quiet?: boolean // defaults to true
+}
+
 export interface EditChangeOpts {
 	newOwner?: string
 	newWorkspace?: string
@@ -665,14 +673,16 @@ export class PerforceContext {
 	 * Get a list of changes in a path since a specific CL
 	 * @return Promise to list of changelists
 	 */
-	changes(path_in: string, since: number, limit?: number, status?: ChangelistStatus, quiet: boolean = true): Promise<Change[]> {
-		const path = since > 0 ? `${path_in}@${since},now` : path_in;
+	changes(path_in: string, opts?: ChangesOpts): Promise<Change[]> {
+		opts = opts || {}
+		const range = '@' + (opts.since ? (opts.until ? `${opts.since},${opts.until}` : `${opts.since},now`) : (opts.until ? `${opts.until}` : ''))
+		const path = path_in + range
 		const args = ['changes', '-l',
-			(status ? `-s${status}` : '-ssubmitted'),
-			...(limit ? [`-m${limit}`] : []),
+			(opts.status ? `-s${opts.status}` : '-ssubmitted'),
+			...(opts.limit ? [`-m${opts.limit}`] : []),
 			path];
 
-		return this.execAndParse(null, args, {quiet}, {
+		return this.execAndParse(null, args, {quiet: opts.quiet||true}, {
 			expected: {change: 'integer', client: 'string', user: 'string', desc: 'string'},
 			optional: {shelved: 'integer', oldChange: 'integer', IsPromoted: 'integer'}
 		}) as Promise<unknown> as Promise<Change[]>
diff --git a/Engine/Extras/RoboMerge/v3/src/robo/nodebot.ts b/Engine/Extras/RoboMerge/v3/src/robo/nodebot.ts
--- a/Engine/Extras/RoboMerge/v3/src/robo/nodebot.ts
+++ b/Engine/Extras/RoboMerge/v3/src/robo/nodebot.ts
@@ -220,7 +220,7 @@ export class NodeBot extends PerforceStatefulBot implements NodeBotInterface {
 				return result
 			}
 			else {
-				return await this.p4.changes(path, startingCl)
+				return await this.p4.changes(path, {since: startingCl})
 			}
 		}
 		catch (err) {
diff --git a/Engine/Extras/RoboMerge/v3/src/robo/trackchange.ts b/Engine/Extras/RoboMerge/v3/src/robo/trackchange.ts
--- a/Engine/Extras/RoboMerge/v3/src/robo/trackchange.ts
+++ b/Engine/Extras/RoboMerge/v3/src/robo/trackchange.ts
@@ -331,7 +331,7 @@ export async function trackChange(robo: RoboMerge, logger: ContextualLogger, ini
                     const externalP4 = PerforceContext.getServerContext(logger, externalTarget.serverID)
                     const streamName = getStreamFromPath(changeToConsider.node!.stream)
                     const targetPath = entry.depotFile.replace(streamName || changeToConsider.node!.stream, externalTarget.stream)
-                    const changes = await externalP4.changes(targetPath, 0)
+                    const changes = await externalP4.changes(targetPath)
                     for (let change of changes) {
                         if (change.desc.includes(clToConsider.cl.toString())) {
                             const newChangeToConsider = {cl: change.change, serverId: externalTarget.serverID}
