Ping and invite the user responsible for a bad branchmap update

diff --git a/Engine/Extras/RoboMerge/v3/src/robo/autobranchupdater.ts b/Engine/Extras/RoboMerge/v3/src/robo/autobranchupdater.ts
--- a/Engine/Extras/RoboMerge/v3/src/robo/autobranchupdater.ts
+++ b/Engine/Extras/RoboMerge/v3/src/robo/autobranchupdater.ts
@@ -4,7 +4,7 @@ import * as fs from 'fs';
 import * as path from 'path';
 import { _setTimeout } from '../common/helper';
 import { ContextualLogger } from '../common/logger';
-import { P4_FORCE, PerforceContext, Workspace } from '../common/perforce';
+import { Change, P4_FORCE, PerforceContext, Workspace } from '../common/perforce';
 import { Bot } from './bot-interfaces';
 import { BranchDefs } from './branchdefs';
 import { GraphBot } from './graphbot';
@@ -39,6 +39,7 @@ export class AutoBranchUpdater implements Bot {
 
 	private readonly abuLogger: ContextualLogger
 	
+	prevLastCl: number
 	lastCl: number
 	lastModifiedTime: number // Only used in devmode
 
@@ -58,6 +59,7 @@ export class AutoBranchUpdater implements Bot {
 		this.filePath = `${config.rootPath}/${graphBot.filename}`
 		this.workspace = config.workspace
 
+		this.prevLastCl = AutoBranchUpdater.initialCl
 		this.lastCl = AutoBranchUpdater.initialCl
 		this.abuLogger = parentLogger.createChild('ABU')
 	}
@@ -127,9 +129,10 @@ export class AutoBranchUpdater implements Bot {
 			const change = await this.p4.latestChange(this.filePath);
 			if (change !== null && change.change > this.lastCl) {
 				await this.p4.sync(this.workspace, `${this.filePath}@${change.change}`);
-	
+
 				// set this to be the last changelist regardless of success - if it failed due to a broken
 				// .json file, the file will have to be recommitted anyway
+				this.prevLastCl = this.lastCl;
 				this.lastCl = change.change;
 				reloadBranchDefs = true;
 			}
@@ -172,8 +175,38 @@ export class AutoBranchUpdater implements Bot {
 
 		const result = await BranchDefs.parseAndValidate(this.abuLogger, branchGraphText)
 		if (!result.branchGraphDef) {
-			// @todo email author of changes!
-			const errText = `Failed to parse/validate branch specs file for ${botname}`
+			const slackMessages = this.graphBot.getSlackMessages()
+			let changes: Change[] = []
+			if (slackMessages && this.prevLastCl != this.lastCl) {
+				changes = await this.p4.changes(this.filePath, {since: this.prevLastCl+1, until: this.lastCl})
+			}
+			let errText = `Failed to parse/validate branch specs file for ${botname}`
+			if (changes.length > 0) {
+				const uniqueUsers = changes.reduce((map, change) => {
+					if (!map.has(change.user)) {
+						map.set(change.user, {
+							changes: [],
+							emailPromise: this.p4.getEmail(change.user)
+						})
+					}
+					map.get(change.user)!.changes.push(change)
+					return map
+				}, new Map<string, { changes: Change[], emailPromise: Promise<any> }>())
+
+				await Promise.all(
+					Array.from(uniqueUsers.values())
+						 .map(userInfo => userInfo.emailPromise
+						 .then(email => email ? slackMessages!.addUserToChannel(email, RoboArgs.slackAlertChannel) : undefined)))
+
+				const userMessages = await Promise.all(
+					Array.from(uniqueUsers.entries()).map(async ([username, userInfo]) => {
+						const pingTarget = await slackMessages!.getPingTarget(userInfo.emailPromise, username)
+						const changeList = userInfo.changes.map(change => `<${this.p4.swarmURL}/changes/${change.change}|CL ${change.change}>`).join(', ')
+						return `${pingTarget} (${changeList})`
+					})
+				)
+				errText += '\nPossibly ' + userMessages.join(', ')
+			}
 			const validationErrorText = result.validationErrors.join('\n')
 			const file: SlackFile = {
 				content: validationErrorText,
diff --git a/Engine/Extras/RoboMerge/v3/src/robo/graphbot.ts b/Engine/Extras/RoboMerge/v3/src/robo/graphbot.ts
--- a/Engine/Extras/RoboMerge/v3/src/robo/graphbot.ts
+++ b/Engine/Extras/RoboMerge/v3/src/robo/graphbot.ts
@@ -511,6 +511,11 @@ export class GraphBot implements BotEventHandler {
 	get settings() {
 		return this._settings;
 	}
+
+	getSlackMessages() {
+		return this.slackMessages;
+	}
+
 	private botlist: Bot[] = []
 	private waitTime?: number
 
