Handle cases where slack postMessage fails (likely due to rate limiting) and we do not post to a thread or response as a result

diff --git a/Engine/Extras/RoboMerge/v3/src/common/request.ts b/Engine/Extras/RoboMerge/v3/src/common/request.ts
--- a/Engine/Extras/RoboMerge/v3/src/common/request.ts
+++ b/Engine/Extras/RoboMerge/v3/src/common/request.ts
@@ -16,10 +16,15 @@ export interface PostArgs extends GetArgs {
 	contentType?: string
 }
 
-export function req(args: GetArgs, method: string) {
+export interface ResponseWithHeaders {
+	body: string
+	headers: http.IncomingHttpHeaders
+}
+
+export function reqWithHeaders(args: GetArgs, method: string) {
 	const IS_POST = method === 'POST' || method === 'PUT'
 
-	return new Promise<string>((done, fail) => {
+	return new Promise<ResponseWithHeaders>((done, fail) => {
 		const url = new _url.URL(args.url);
 
 		let response = '';
@@ -41,7 +46,7 @@ export function req(args: GetArgs, method: string) {
 
 		const callback = (res: http.IncomingMessage) => {
 			res.on('data', chunk => response += chunk)
-			res.on('end', () => done(response))
+			res.on('end', () => done({body: response, headers: res.headers}))
 		}
 
 		let req: http.ClientRequest = url.protocol === 'http:'
@@ -56,6 +61,11 @@ export function req(args: GetArgs, method: string) {
 	});
 }
 
+export async function req(args: GetArgs, method: string) {
+	const response = await reqWithHeaders(args, method)
+	return response.body
+}
+
 export function get(args: GetArgs) {
 	return req(args, 'GET')
 }
@@ -64,6 +74,10 @@ export function post(args: PostArgs) {
 	return req(args, 'POST')
 }
 
+export function postWithHeaders(args: PostArgs) {
+	return reqWithHeaders(args, 'POST')
+}
+
 export function put(args: PostArgs) {
 	return req(args, 'PUT')
 }
diff --git a/Engine/Extras/RoboMerge/v3/src/robo/slack.ts b/Engine/Extras/RoboMerge/v3/src/robo/slack.ts
--- a/Engine/Extras/RoboMerge/v3/src/robo/slack.ts
+++ b/Engine/Extras/RoboMerge/v3/src/robo/slack.ts
@@ -56,6 +56,7 @@ export interface SlackMessage {
 	icon_emoji?: string
 	pretext?: string
 	footer?: string
+	thread_ts?: string
 
 	// An array because the Slack API expects it to be
 	attachments?: SlackAttachment[]
@@ -84,7 +85,7 @@ enum PostFormat {
 	URL
 }
 
-const MAIN_MESSAGE_FIELDS = new Set(['username', 'icon_emoji', 'channel', 'target', 'cl']);
+const MAIN_MESSAGE_FIELDS = new Set(['username', 'icon_emoji', 'channel', 'target', 'cl', 'thread_ts']);
 
 export class Slack {
 	constructor(private channel: SlackChannel, private domain: string, private readonly logger: ContextualLogger) {
@@ -105,14 +106,41 @@ export class Slack {
 		return (await this.get('chat.getPermalink', {channel, message_ts: thread_ts})).permalink
 	}
 
-	async postMessage(message: SlackMessage) {
-		return (await this.post('chat.postMessage', this.makeArgs(message))).ts
+	async postMessage(message: SlackMessage, retryCount: number = 0): Promise<string> {
+		const MAX_RETRIES = 3
+		const result = await this.post('chat.postMessage', this.makeArgs(message), true)
+		if (result.ok) {
+			return result.ts
+		}
+		else if (result.error == "rate_limited" ||
+				 result.error == "request_timeout" ||
+				 result.error == "service_unavailable") {
+
+			if (retryCount < MAX_RETRIES) {
+				// Exponential backoff: 1s, 2s, 4s
+				const retrySeconds = Math.pow(2, retryCount)
+
+				this.logger.warn(`chat.postMessage failed with ${result.error}. Retrying after ${retrySeconds} seconds (attempt ${retryCount + 1}/${MAX_RETRIES})`)
+
+				await new Promise(resolve => setTimeout(resolve, retrySeconds * 1000))
+
+				// Retry the request
+				return this.postMessage(message, retryCount + 1)
+			}
+
+			// Max retries exceeded
+			const errMsg = `Slack chat.postMessage failed after ${MAX_RETRIES} retries with error: ${result.error}`
+			this.logger.error(errMsg)
+			throw new Error(errMsg)
+		}
+		else {
+			throw new Error(`Slack chat.postMessage error ${result.error}`)
+		}
 	}
 
 	reply(thread_ts: string, message: SlackMessage) {
-		const args = this.makeArgs(message)
-		args.thread_ts = thread_ts
-		return this.post('chat.postMessage', args)
+		message.thread_ts = thread_ts
+		return this.postMessage(message)
 	}
 
 	update(ts: string, message: SlackMessage) {
@@ -171,23 +199,36 @@ export class Slack {
 		return result.channel ? result.channel.id : null
 	}
 
-	/*private*/ async post_user(userToken: string, command: string, args: any, canFail? : boolean, format?: PostFormat) {
-		const resultJson = await request.post({
+	/*private*/ async post_user(userToken: string, command: string, args: any, canFail? : boolean, format?: PostFormat, retryCount: number = 0): Promise<any> {
+		const MAX_RETRIES = 3
+		const response = await request.postWithHeaders({
 			url: this.domain + '/api/' + command,
 			body: format == PostFormat.URL ? new URLSearchParams(Object.entries(args)).toString() : JSON.stringify(args),
 			headers: {Authorization: 'Bearer ' + userToken},
 			contentType: format == PostFormat.URL ? 'application/x-www-form-urlencoded' : 'application/json; charset=utf-8'
 		})
 		try {
-			const result = JSON.parse(resultJson)
+			const result = JSON.parse(response.body)
 			if (result.ok || canFail) {
 				return result
 			}
+
+			if (result.error === 'ratelimited' && retryCount < MAX_RETRIES) {
+				const retryAfter = response.headers['retry-after']
+				const retrySeconds = retryAfter ? parseInt(retryAfter as string, 10) : 1
+
+				this.logger.warn(`${command} rate limited. Retrying after ${retrySeconds} seconds (attempt ${retryCount + 1}/${MAX_RETRIES})`)
+
+				await new Promise(resolve => setTimeout(resolve, retrySeconds * 1000))
+
+				// Retry the request
+				return this.post_user(userToken, command, args, canFail, format, retryCount + 1)
+			}
 		}
 		catch {
 		}
 
-		this.logger.error(`${command} generated:\n\t${resultJson}`)
+		this.logger.error(`${command} generated:\n\t${response.body}`)
 		return {ok: false}
 	}
 
