can fstat multiple files at once

diff --git a/Engine/Extras/RoboMerge/v3/src/common/perforce.ts b/Engine/Extras/RoboMerge/v3/src/common/perforce.ts
--- a/Engine/Extras/RoboMerge/v3/src/common/perforce.ts
+++ b/Engine/Extras/RoboMerge/v3/src/common/perforce.ts
@@ -2005,9 +2005,12 @@ export class PerforceContext {
 		return result.length == 0
 	}
 
-	async fstat(roboWorkspace: RoboWorkspace, depotPath: string, quiet?: boolean) {
+	async fstat(roboWorkspace: RoboWorkspace, depotPath: string|string[], quiet?: boolean) {
+		if (typeof depotPath === "string") {
+			depotPath = [depotPath]
+		}
 		try {
-			return await this.execAndParse(roboWorkspace, ['fstat', depotPath], {quiet})
+			return await this.execAndParse(roboWorkspace, ['fstat', ...depotPath], {quiet})
 		}
 		catch (reason) {
 			if (!isExecP4Error(reason)) {
