Instead of entirely disabling stomp on large changelists instead drop in to a limited info mode that does not provide as much information but still allows the stomp
Reduce the stomp large changelist size to 20000

diff --git a/Engine/Extras/RoboMerge/v3/public/js/stomp.js b/Engine/Extras/RoboMerge/v3/public/js/stomp.js
--- a/Engine/Extras/RoboMerge/v3/public/js/stomp.js
+++ b/Engine/Extras/RoboMerge/v3/public/js/stomp.js
@@ -102,11 +102,20 @@ function stompVerify() {
                 // Visualize the json
                 visualizeStompVerification(requestedBranchCl, stompJson)
 
-                if (stompJson.validRequest) {
+                if (stompJson.requestResult !== 'Invalid') {
                     $('#afterVerificationResultText').html('Stomp Verification Complete')
-                    $('#afterVerificationResultText').append(
-                        $('<h4 style="text-align: center">').text("Use the button below to proceed with the stomp operation")
-                    )
+
+                    if (stompJson.requestResult === 'LimitedInfo') {
+                        $('#afterVerificationResultText').append(
+                            $('<h4 class="text-center text-warning">').text(
+                                "Limited Information Mode - Review Warnings Before Proceeding"
+                            )
+                        )
+                    } else {
+                        $('#afterVerificationResultText').append(
+                            $('<h4 style="text-align: center">').text("Use the button below to proceed with the stomp operation")
+                        )
+                    }
                 } else {
                     $('#afterVerificationResultText').html('<span><i class="fas fa-exclamation-triangle"></i></span> Stomp Verification Returned Issues.')
                     $('#afterVerificationResultText').css('color', 'red');
@@ -115,7 +124,7 @@ function stompVerify() {
                 transitionPostVerification()
 
                 // If we're a valid request, add form buttons to formally request a stomp
-                if (stompJson.validRequest) {
+                if (stompJson.requestResult !== 'Invalid') {
                     const formButtonDiv = $('#formButtons')
 
                     // Return to Robomerge homepage
@@ -238,7 +247,20 @@ function visualizeChangelist(changelist, changelistData, swarmURL) {
     return changelistDiv
 }
 
+function displayLimitedInfoModeWarning(requestedBranchCl, stompJson) {
+    const linkToSwarm = makeClLink(requestedBranchCl, stompJson.swarmURL, `CL ${requestedBranchCl}`)
+
+    $('#limitedInfoModeWarning').removeClass('initiallyHidden')
+    $('#limitedInfoClMessage').html(`${linkToSwarm} exceeds the analysis limit.`)
+}
+
 function visualizeStompVerification(requestedBranchCl, stompJson) {
+    // Check for limited info mode
+    if (stompJson.requestResult === 'LimitedInfo') {
+        displayLimitedInfoModeWarning(requestedBranchCl, stompJson)
+        return
+    }
+
     // For visualization, sort data by stomped revision
     // (JSON data is sorted by file)
     let clDict = {}
diff --git a/Engine/Extras/RoboMerge/v3/public/stomp.html b/Engine/Extras/RoboMerge/v3/public/stomp.html
--- a/Engine/Extras/RoboMerge/v3/public/stomp.html
+++ b/Engine/Extras/RoboMerge/v3/public/stomp.html
@@ -86,7 +86,37 @@
                     <p>If you require assistance, please <strong><a href="/contact">Contact Us</a></strong> for Robomerge help.</p>
                 </div> 
             </div>
-        
+
+        <!-- Display warning for limited info mode on large changelists -->
+        <div id="limitedInfoModeWarning" class="card card-single border-warning centered-text initiallyHidden">
+            <div class="card-header bg-warning text-white">
+                <span><i class="fas fa-exclamation-triangle"></i></span>
+                <span class="sr-only">Warning:</span>
+                Large Changelist - Limited Information Available
+            </div>
+            <div class="card-body">
+                <p id="limitedInfoClMessage" class="text-center"></p>
+                <p class="text-center">Stomp can proceed, but detailed revision information is unavailable.</p>
+
+                <h6>What this means:</h6>
+                <ul class="text-left">
+                    <li>Robomerge has verified that all remaining conflicts are binary files and safe to stomp</li>
+                    <li>Detailed information about which CLs/authors will be affected is <strong>NOT available</strong></li>
+                    <li>Affected authors will <strong>NOT be automatically notified</strong> via #fyi</li>
+                    <li>You will <strong>NOT see a preview</strong> of stomped revisions</li>
+                </ul>
+
+                <h6>Before proceeding:</h6>
+                <ul class="text-left">
+                    <li>Review the file history in P4V to identify affected changes</li>
+                    <li>Manually notify affected authors if needed</li>
+                    <li>Ensure this stomp operation is appropriate for your workflow</li>
+                </ul>
+
+                <p class="text-center mt-3">If you're unsure, please reach out for assistance.</p>
+            </div>
+        </div>
+
         <!-- Dynamically generated div to hold visualization of stomp verification JSON -->
         <div id="resultVisualization"></div>
     </div>
diff --git a/Engine/Extras/RoboMerge/v3/src/robo/branch-interfaces.ts b/Engine/Extras/RoboMerge/v3/src/robo/branch-interfaces.ts
--- a/Engine/Extras/RoboMerge/v3/src/robo/branch-interfaces.ts
+++ b/Engine/Extras/RoboMerge/v3/src/robo/branch-interfaces.ts
@@ -48,8 +48,14 @@ export interface StompVerificationFile {
 	targetFileName : string
 }
 
+export enum StompRequestResult {
+	'Invalid',
+	'Valid',
+	'LimitedInfo'
+}
+
 export interface StompVerification extends OperationResult {
-	validRequest?: boolean
+	requestResult?: keyof typeof StompRequestResult
 	nonBinaryFilesResolved?: boolean // Warning to the user that non-binary files are included in changelist, but passed the first resolve
 	svFiles?: StompVerificationFile[] // Array of remaining files
 
@@ -65,7 +71,7 @@ export function GetStompVerificationResult(verifyResult: StompVerification): Ope
 		return { success: false, message: `stompChanges Request Failed Due to Bad Verification: ${verifyResult.message}` }
 	}
 
-	if (!verifyResult.validRequest) {
+	if (verifyResult.requestResult === 'Invalid') {
 		return { success: false, message: `stompChanges Request Cannot Be Completed: ${verifyResult.message}` }
 	}
 
@@ -256,7 +262,10 @@ export interface MergeConflictAdditionalInformation {
 	resolveResult: ResolveResultAdditionalInfo,
 	describeResult?: DescribeResult,
 	disableStompReason?: string,
-	client?: string
+	client?: string,
+
+	// Limited information mode for large changelists
+	limitedInfoMode?: boolean
 }
 
 export function isA_MergeConflictAdditionalInformation(obj: any): obj is MergeConflictAdditionalInformation {
@@ -266,7 +275,8 @@ export function isA_MergeConflictAdditionalInformation(obj: any): obj is MergeCo
     typeof obj.resolveResult === "object" &&
     (obj.describeResult === undefined || typeof obj.describeResult === "object") &&
     (obj.disableStompReason === undefined || typeof obj.disableStompReason === "string") &&
-    (obj.client === undefined || typeof obj.client === "string")
+    (obj.client === undefined || typeof obj.client === "string") &&
+    (obj.limitedInfoMode === undefined || typeof obj.limitedInfoMode === "boolean")
   )
 }
 
diff --git a/Engine/Extras/RoboMerge/v3/src/robo/edgebot.ts b/Engine/Extras/RoboMerge/v3/src/robo/edgebot.ts
--- a/Engine/Extras/RoboMerge/v3/src/robo/edgebot.ts
+++ b/Engine/Extras/RoboMerge/v3/src/robo/edgebot.ts
@@ -825,8 +825,8 @@ class EdgeBotImpl extends PerforceStatefulBot {
 				}
 			}
 			if (failure.additionalInfo.describeResult.entries.length > RoboArgs.maxStompFileCount) {
+				failure.additionalInfo.limitedInfoMode = true
 				failure.additionalInfo.describeResult = undefined
-				failure.additionalInfo.disableStompReason = "Changelist too large"
 			}
 
 			const conflicts = this.analyzeConflict(result.getConflicts())
diff --git a/Engine/Extras/RoboMerge/v3/src/robo/ipc.ts b/Engine/Extras/RoboMerge/v3/src/robo/ipc.ts
--- a/Engine/Extras/RoboMerge/v3/src/robo/ipc.ts
+++ b/Engine/Extras/RoboMerge/v3/src/robo/ipc.ts
@@ -554,14 +554,14 @@ export class IPC {
 			// Attempt to stomp changes
 			let stompVerification = await bot.verifyStomp(cl, target)
 			if (stompVerification.success) {
-				return { 
+				return {
 					statusCode: 200,
 					message: JSON.stringify({
 						message: stompVerification.message,
 						nonBinaryFilesResolved: stompVerification.nonBinaryFilesResolved,
 						remainingAllBinary: stompVerification.remainingAllBinary,
 						files: stompVerification.svFiles,
-						validRequest: stompVerification.validRequest,
+						requestResult: stompVerification.requestResult,
 						swarmURL: stompVerification.swarmURL
 					} )
 				}
diff --git a/Engine/Extras/RoboMerge/v3/src/robo/nodebot.ts b/Engine/Extras/RoboMerge/v3/src/robo/nodebot.ts
--- a/Engine/Extras/RoboMerge/v3/src/robo/nodebot.ts
+++ b/Engine/Extras/RoboMerge/v3/src/robo/nodebot.ts
@@ -719,6 +719,10 @@ export class NodeBot extends PerforceStatefulBot implements NodeBotInterface {
 	async _evaluateConflictForStomp(changeCl: number, targetBranch: Readonly<EditableBranch>, mergeConflictInfo: MergeConflictAdditionalInformation) : Promise<StompVerification>
 	{
 		try {
+			// Check for limited info mode first
+			if (mergeConflictInfo.limitedInfoMode) {
+				return await this._evaluateConflictForStompLimitedMode(changeCl, mergeConflictInfo)
+			}
 
 			if (!mergeConflictInfo.describeResult) {
 				return this.stompFailureResult(mergeConflictInfo.disableStompReason || "DescribeResult missing for unknown reason")
@@ -865,12 +869,12 @@ export class NodeBot extends PerforceStatefulBot implements NodeBotInterface {
 
 			if (!remainingAllBinary) {
 				returnResult.message = "Invalid Request: Not all remaining files are binary"
-				returnResult.validRequest = false
+				returnResult.requestResult = 'Invalid'
 			}
 			else if (stompedRevisionsCalculationIssues) {
 				returnResult.message = "Valid Request, but encountered issues calculating stomped revisions on some files"
-				returnResult.validRequest = true
-				
+				returnResult.requestResult = 'Valid'
+
 				if (this.tickJournal) {
 					// Alert our analytics that we ran into this issue
 					++this.tickJournal.stompedRevisionsCalculationIssues
@@ -878,7 +882,7 @@ export class NodeBot extends PerforceStatefulBot implements NodeBotInterface {
 			}
 			else {
 				returnResult.message = "Valid Request"
-				returnResult.validRequest = true
+				returnResult.requestResult = 'Valid'
 			}
 
 			// Finally return the results
@@ -886,7 +890,60 @@ export class NodeBot extends PerforceStatefulBot implements NodeBotInterface {
 			
 		} catch (error) {
 			return this.stompFailureResult(`Encountered error processing ${changeCl}: ${error}`)
-		} 
+		}
+	}
+
+	private async _evaluateConflictForStompLimitedMode(changeCl: number, mergeConflictInfo: MergeConflictAdditionalInformation): Promise<StompVerification> {
+		try {
+			// Get resolve result to check for conflicts
+			const resolveResult = Object.assign(new ResolveResultAdditionalInfo(), mergeConflictInfo.resolveResult)
+
+			if (!resolveResult.hasConflict()) {
+				return this.stompFailureResult(`No conflict found resolving CL ${changeCl}`)
+			}
+			if (!resolveResult.successfullyParsedConflicts) {
+				return this.stompFailureResult(`Encountered error enumerating conflicting files after integration attempt: ${resolveResult.parsingError}`)
+			}
+
+			if (!mergeConflictInfo.client) {
+				return this.stompFailureResult(`Client workspace information missing for CL ${changeCl}`)
+			}
+
+			const remainingFiles = resolveResult.getConflicts()
+
+			const BATCH_SIZE = 1000
+			let remainingAllBinary = true
+			for (let i = 0; i < remainingFiles.length && remainingAllBinary; i += BATCH_SIZE) {
+				const batch = remainingFiles.slice(i, i + BATCH_SIZE)			
+				const clientFilePaths = batch.map(f => f.clientFile)
+				const whereResults = await this.p4.where(mergeConflictInfo.client!, clientFilePaths)
+				const depotPaths = whereResults.map(r => r.depotFile).filter(f => f) as string[]
+				const fstatResults = await this.p4.fstat(mergeConflictInfo.client!, depotPaths, true)
+				
+				// Check each result directly
+				for (const fstatResult of fstatResults) {
+					const stompable = fstatResult.headType.startsWith('binary') || isStompableNonBinary(fstatResult.depotFile)
+					if (!stompable) {
+						remainingAllBinary = false
+						break
+					}
+				}
+			}
+
+			return {
+				success: true,
+				requestResult: remainingAllBinary ? 'LimitedInfo' : 'Invalid',
+				remainingAllBinary,
+				message: remainingAllBinary
+					? "Valid Request (Limited Information Mode - Changelist exceeds file count limit)"
+					: "Invalid Request: Not all remaining files are binary",
+				svFiles: [],  // Empty - no detailed file info
+				nonBinaryFilesResolved: false,  // Unknown in limited mode
+				swarmURL: this.p4.swarmURL
+			}
+		} catch (error) {
+			return this.stompFailureResult(`Encountered error processing ${changeCl} in limited mode: ${error}`)
+		}
 	}
 
 	// TODO: Support branch/delete merges
@@ -959,20 +1016,31 @@ export class NodeBot extends PerforceStatefulBot implements NodeBotInterface {
 			return operationResult
 		}
 
-		// Before we submit, add an #fyi for all the stomped revision authors in the changelist
-		let fyiSet = new Set<string>([owner.toLowerCase()])
-		for (const file of verifyResult.svFiles!) {
-			if (file.stompedRevisions) {
-				for (const rev of file.stompedRevisions) {
-					fyiSet.add(rev.author.toLowerCase())
+		// Build description based on mode
+		let description = ''
+
+		if (verifyResult.requestResult === 'LimitedInfo') {
+			this.nodeBotLogger.warn(`Stomp in limited mode: CL ${changeCl} -> ${targetBranchName}`)
+			description = `\nStomp Operation performed without detailed revision tracking due to large changelist size.\n` +
+						 `Authors of stomped changes were not automatically identified. Please verify affected changes manually.`
+		}
+		else {
+			// Normal mode: build #fyi from stomped revisions
+			const fyiSet = new Set<string>([owner.toLowerCase()])
+			for (const file of verifyResult.svFiles!) {
+				if (file.stompedRevisions) {
+					for (const rev of file.stompedRevisions) {
+						fyiSet.add(rev.author.toLowerCase())
+					}
 				}
 			}
+			description = `\nStomp Operation may have overwritten commits by these authors:\n#fyi ${Array.from(fyiSet).join(', ')}`
 		}
 
 		this.nodeBotLogger.info(`Stomp Request for cl ${changeCl} by ${owner} validated, submitting change through reconsider system`)
 		this.reconsider(owner, changeCl, {
-			additionalFlags: ['stomp'], targetBranchName, 
-			description: `\nStomp Operation may have overwritten commits by these authors:\n#fyi ${Array.from(fyiSet).join(', ')}`
+			additionalFlags: ['stomp'], targetBranchName,
+			description
 		})
 
 		if (this.tickJournal) {
diff --git a/Engine/Extras/RoboMerge/v3/src/robo/robo.ts b/Engine/Extras/RoboMerge/v3/src/robo/robo.ts
--- a/Engine/Extras/RoboMerge/v3/src/robo/robo.ts
+++ b/Engine/Extras/RoboMerge/v3/src/robo/robo.ts
@@ -168,7 +168,7 @@ const COMMAND_LINE_ARGS: {[param: string]: Arg<any>} = {
 		parse: (str: string) => {
 			return parseInt(str)
 		},
-		dflt: 30000
+		dflt: 20000
 	},
 
 	maxReconsiderFiles: {
diff --git a/Engine/Extras/RoboMerge/v3/src/robo/settings.ts b/Engine/Extras/RoboMerge/v3/src/robo/settings.ts
--- a/Engine/Extras/RoboMerge/v3/src/robo/settings.ts
+++ b/Engine/Extras/RoboMerge/v3/src/robo/settings.ts
@@ -356,7 +356,7 @@ export class Settings {
 			if (isA_MergeConflictAdditionalInformation(additionalInfo)) {
 				if (additionalInfo?.describeResult) {
 					if (!trim || additionalInfo.describeResult.entries.length > RoboArgs.maxStompFileCount) {
-						additionalInfo.disableStompReason = "Changelist too large"
+						additionalInfo.limitedInfoMode = true
 						additionalInfo.describeResult = undefined
 					}
 				}
