Fix unlock looking at the wrong files property

diff --git a/Engine/Extras/RoboMerge/v3/public/js/unlock.js b/Engine/Extras/RoboMerge/v3/public/js/unlock.js
--- a/Engine/Extras/RoboMerge/v3/public/js/unlock.js
+++ b/Engine/Extras/RoboMerge/v3/public/js/unlock.js
@@ -219,7 +219,7 @@ function visualizeAuthor(author, files) {
 
 function visualizeUnlockVerification(requestedBranchCl, unlockJson) {
     let authorDict = {}
-    for (const file of unlockJson.files) {
+    for (const file of unlockJson.lockedFiles) {
 
         if (!authorDict[file.user]) {
             authorDict[file.user] = {
