Fix functional tests not looking at the correct properties.
Fix missed case in stomp frontend

diff --git a/Engine/Extras/RoboMerge/v3/functional_tests/framework/src/framework.ts b/Engine/Extras/RoboMerge/v3/functional_tests/framework/src/framework.ts
--- a/Engine/Extras/RoboMerge/v3/functional_tests/framework/src/framework.ts
+++ b/Engine/Extras/RoboMerge/v3/functional_tests/framework/src/framework.ts
@@ -815,7 +815,7 @@ export abstract class FunctionalTest {
 			checkSlack()
 		])
 
-		if (!verifyResult.validRequest)
+		if (verifyResult.requestResult !== 'Valid')
 		{
 			this.warn(verifyResult.message)
 			throw new Error('Stomp verify returned unexpected values')
diff --git a/Engine/Extras/RoboMerge/v3/functional_tests/framework/src/tests/confirm-binary-stomp.ts b/Engine/Extras/RoboMerge/v3/functional_tests/framework/src/tests/confirm-binary-stomp.ts
--- a/Engine/Extras/RoboMerge/v3/functional_tests/framework/src/tests/confirm-binary-stomp.ts
+++ b/Engine/Extras/RoboMerge/v3/functional_tests/framework/src/tests/confirm-binary-stomp.ts
@@ -128,7 +128,7 @@ export class ConfirmBinaryStomp extends FunctionalTest {
 		// verify stomp
 		let result = await this.verifyAndPerformStomp('Main', branch, this.testName)
 
-		let files = result.verify.files
+		let files = result.verify.svFiles
 
 		if (files[0].targetFileName === `//${this.testName}/${branch}/${COLLECTION_FILENAME}` &&
 			files[1].targetFileName === `//${this.testName}/${branch}/${ASSET_FILENAME}`) {
diff --git a/Engine/Extras/RoboMerge/v3/functional_tests/framework/src/tests/reject-text-conflict-stomp.ts b/Engine/Extras/RoboMerge/v3/functional_tests/framework/src/tests/reject-text-conflict-stomp.ts
--- a/Engine/Extras/RoboMerge/v3/functional_tests/framework/src/tests/reject-text-conflict-stomp.ts
+++ b/Engine/Extras/RoboMerge/v3/functional_tests/framework/src/tests/reject-text-conflict-stomp.ts
@@ -108,7 +108,7 @@ export class RejectTextConflictStomp extends FunctionalTest {
 		const edgeState = await this.getEdgeState('Development', 'Main')
 		const verifyResult = await this.verifyStompRequest('Development', 'Main', edgeState)
 
-		if (!verifyResult.validRequest && !verifyResult.nonBinaryFilesResolved && !verifyResult.remainingAllBinary) {
+		if (verifyResult.requestResult !== 'Valid' && !verifyResult.nonBinaryFilesResolved && !verifyResult.remainingAllBinary) {
 			this.info('Success! Stomp verify detected files remaining after the automatic merge, which is correct for this functional test.')
 		}
 		else {
diff --git a/Engine/Extras/RoboMerge/v3/functional_tests/framework/src/tests/stomp-with-add.ts b/Engine/Extras/RoboMerge/v3/functional_tests/framework/src/tests/stomp-with-add.ts
--- a/Engine/Extras/RoboMerge/v3/functional_tests/framework/src/tests/stomp-with-add.ts
+++ b/Engine/Extras/RoboMerge/v3/functional_tests/framework/src/tests/stomp-with-add.ts
@@ -75,7 +75,7 @@ export class StompWithAdd extends FunctionalTest {
 		// verify stomp
 		let result = await this.verifyAndPerformStomp('Main', branch, this.testName)
 
-		let files = result.verify.files
+		let files = result.verify.svFiles
 
 		for (const file of files) {
 			const isEditAsset = file.targetFileName.indexOf(EDIT_ASSET_FILENAME) >= 0
diff --git a/Engine/Extras/RoboMerge/v3/public/js/stomp.js b/Engine/Extras/RoboMerge/v3/public/js/stomp.js
--- a/Engine/Extras/RoboMerge/v3/public/js/stomp.js
+++ b/Engine/Extras/RoboMerge/v3/public/js/stomp.js
@@ -265,7 +265,7 @@ function visualizeStompVerification(requestedBranchCl, stompJson) {
     // (JSON data is sorted by file)
     let clDict = {}
     let encounteredStompedRevisionsCalculationIssues = false
-    for (const file of stompJson.files) {
+    for (const file of stompJson.svFiles) {
         const swarmFileLink = makeSwarmFileLink(file.targetFileName, stompJson.swarmURL)
 
         // Process resolved files
@@ -338,7 +338,7 @@ function visualizeStompVerification(requestedBranchCl, stompJson) {
         )
     } 
 
-    if (stompJson.validRequest) {
+    if (stompJson.requestResult !== 'Invalid') {
         if (stompJson.nonBinaryFilesResolved) {
             // Enable warning if we're doing a mixed merge of binary and non-binary files
             $('#mixedMergeWarning').removeClass('initiallyHidden')
