Additional cleanup of stomp verification and functional tests

diff --git a/Engine/Extras/RoboMerge/v3/functional_tests/framework/src/tests/confirm-binary-stomp.ts b/Engine/Extras/RoboMerge/v3/functional_tests/framework/src/tests/confirm-binary-stomp.ts
--- a/Engine/Extras/RoboMerge/v3/functional_tests/framework/src/tests/confirm-binary-stomp.ts
+++ b/Engine/Extras/RoboMerge/v3/functional_tests/framework/src/tests/confirm-binary-stomp.ts
@@ -128,7 +128,7 @@ export class ConfirmBinaryStomp extends FunctionalTest {
 		// verify stomp
 		let result = await this.verifyAndPerformStomp('Main', branch, this.testName)
 
-		let files = result.verify.svFiles
+		let files = result.verify.files
 
 		if (files[0].targetFileName === `//${this.testName}/${branch}/${COLLECTION_FILENAME}` &&
 			files[1].targetFileName === `//${this.testName}/${branch}/${ASSET_FILENAME}`) {
diff --git a/Engine/Extras/RoboMerge/v3/functional_tests/framework/src/tests/stomp-with-add.ts b/Engine/Extras/RoboMerge/v3/functional_tests/framework/src/tests/stomp-with-add.ts
--- a/Engine/Extras/RoboMerge/v3/functional_tests/framework/src/tests/stomp-with-add.ts
+++ b/Engine/Extras/RoboMerge/v3/functional_tests/framework/src/tests/stomp-with-add.ts
@@ -75,7 +75,7 @@ export class StompWithAdd extends FunctionalTest {
 		// verify stomp
 		let result = await this.verifyAndPerformStomp('Main', branch, this.testName)
 
-		let files = result.verify.svFiles
+		let files = result.verify.files
 
 		for (const file of files) {
 			const isEditAsset = file.targetFileName.indexOf(EDIT_ASSET_FILENAME) >= 0
diff --git a/Engine/Extras/RoboMerge/v3/public/js/stomp.js b/Engine/Extras/RoboMerge/v3/public/js/stomp.js
--- a/Engine/Extras/RoboMerge/v3/public/js/stomp.js
+++ b/Engine/Extras/RoboMerge/v3/public/js/stomp.js
@@ -265,7 +265,7 @@ function visualizeStompVerification(requestedBranchCl, stompJson) {
     // (JSON data is sorted by file)
     let clDict = {}
     let encounteredStompedRevisionsCalculationIssues = false
-    for (const file of stompJson.svFiles) {
+    for (const file of stompJson.files) {
         const swarmFileLink = makeSwarmFileLink(file.targetFileName, stompJson.swarmURL)
 
         // Process resolved files
diff --git a/Engine/Extras/RoboMerge/v3/src/robo/ipc.ts b/Engine/Extras/RoboMerge/v3/src/robo/ipc.ts
--- a/Engine/Extras/RoboMerge/v3/src/robo/ipc.ts
+++ b/Engine/Extras/RoboMerge/v3/src/robo/ipc.ts
@@ -374,7 +374,7 @@ export class IPC {
 	}
 
 	private logError(msg: string) {
-		if (RoboArgs.isRunningFunctionalTests()) {
+		if (RoboArgs.isRunningFunctionalTests) {
 			this.ipcLogger.info(msg)
 		} else {
 			this.ipcLogger.error(msg)
