Destory and recreate branchmap and gate workspaces to ensure they always point at the right root folder

diff --git a/Engine/Extras/RoboMerge/v3/src/common/perforce.ts b/Engine/Extras/RoboMerge/v3/src/common/perforce.ts
--- a/Engine/Extras/RoboMerge/v3/src/common/perforce.ts
+++ b/Engine/Extras/RoboMerge/v3/src/common/perforce.ts
@@ -891,6 +891,27 @@ export class PerforceContext {
 		return null
 	}
 
+	async deleteWorkspace(workspaceName: string, edgeServerAddress?: string) {
+		let result
+		try {
+			await this.execAndParse(null, ['client', '-d', workspaceName], {serverAddress: edgeServerAddress});
+		}
+		catch (reason) {
+			if (!isExecP4Error(reason)) {
+				throw reason
+			}
+			const [err, output] = reason
+
+			// delete succeeds with an error because ... sigh
+			if (!output.trim().startsWith('Client') ||
+				!output.trim().endsWith('deleted.')) {
+				throw err
+			}
+			result = output;			
+		}
+		return result
+	}
+
 	async reloadWorkspace(workspaceName: string, edgeServerAddress?: string) {
 		let result
 		try {
diff --git a/Engine/Extras/RoboMerge/v3/src/robo/gate.ts b/Engine/Extras/RoboMerge/v3/src/robo/gate.ts
--- a/Engine/Extras/RoboMerge/v3/src/robo/gate.ts
+++ b/Engine/Extras/RoboMerge/v3/src/robo/gate.ts
@@ -474,15 +474,13 @@ export class Gate {
 		const p4 = this.context.p4!
 	
 		const workspaceQueryResult = await p4.find_workspace_by_name(workspace.name, {includeUnloaded: true})
-		if (workspaceQueryResult.length === 0) {
-			await p4.newWorkspace(workspace.name, {
-				Stream: stream,
-				Root: workspace.directory
-			})
-		}
-		else if (workspaceQueryResult[0].IsUnloaded) {
-			await p4.reloadWorkspace(workspace.name)
+		if (workspaceQueryResult.length === 1) {
+			await p4.deleteWorkspace(workspace.name)
 		}
+		await p4.newWorkspace(workspace.name, {
+			Stream: stream,
+			Root: workspace.directory
+		})
 	
 		const syncResult = await p4.sync(workspace, lastGoodCLPath, {opts:[P4_FORCE]})
 		if (!syncResult || syncResult.length != 1) {
diff --git a/Engine/Extras/RoboMerge/v3/src/robo/robo.ts b/Engine/Extras/RoboMerge/v3/src/robo/robo.ts
--- a/Engine/Extras/RoboMerge/v3/src/robo/robo.ts
+++ b/Engine/Extras/RoboMerge/v3/src/robo/robo.ts
@@ -451,17 +451,13 @@ async function init(logger: ContextualLogger) {
 			devMode: args.devMode
 		}
 
-		// Ensure we have a workspace for branch specs
+		// Ensure we have a workspace for branch specs and that it points to the correct root path
 		const workspace = await p4.find_workspace_by_name(args.branchSpecsWorkspace, {includeUnloaded: true})
-		if (workspace.length === 0) {
-			logger.info(`Cannot find workspace ${args.branchSpecsWorkspace}, creating a new one.`)
-
-			if ((await lookupStream(args.branchSpecsRootPath)).length > 0) {
-				await p4.newBranchSpecWorkspace(autoUpdaterConfig.workspace, args.branchSpecsRootPath)
-			}
+		if (workspace.length === 1) {
+			await p4.deleteWorkspace(args.branchSpecsWorkspace)
 		}
-		else if (workspace[0].IsUnloaded) {
-			await p4.reloadWorkspace(args.branchSpecsWorkspace)
+		if ((await lookupStream(args.branchSpecsRootPath)).length > 0) {
+			await p4.newBranchSpecWorkspace(autoUpdaterConfig.workspace, args.branchSpecsRootPath)
 		}
 
 		// make sure we've got the latest branch specs
