Preserve target tab through okta login redirection

diff --git a/Engine/Extras/RoboMerge/v3/public/js/boilerplate.js b/Engine/Extras/RoboMerge/v3/public/js/boilerplate.js
--- a/Engine/Extras/RoboMerge/v3/public/js/boilerplate.js
+++ b/Engine/Extras/RoboMerge/v3/public/js/boilerplate.js
@@ -7,6 +7,17 @@
  * COMMON FUNCTIONS *
  ********************/
 
+// Hash preservation for authentication flow
+function preserveHashForAuth() {
+	if (window.location.hash) {
+		try {
+			sessionStorage.setItem('robomerge_auth_hash', window.location.hash);
+		} catch (e) {
+			console.warn('Failed to preserve hash in sessionStorage:', e);
+		}
+	}
+}
+
 // Global user information map
 let robomergeUser
 let helpPageURL = '/help'
@@ -1994,4 +2005,7 @@ if (!onLoginPage) {
 			}
 		}
 	});
+
+	// Preserve hash before potential auth redirect
+	preserveHashForAuth();
 }
diff --git a/Engine/Extras/RoboMerge/v3/public/js/login.js b/Engine/Extras/RoboMerge/v3/public/js/login.js
--- a/Engine/Extras/RoboMerge/v3/public/js/login.js
+++ b/Engine/Extras/RoboMerge/v3/public/js/login.js
@@ -26,6 +26,15 @@ $(async () => {
 	const $oktaForm = $('.login-form .okta-form');
 	const $oktaError = $('#okta-login-error');
 
+	// Preserve any hash from URL on login page load
+	if (window.location.hash) {
+		try {
+			sessionStorage.setItem('robomerge_auth_hash', window.location.hash);
+		} catch (e) {
+			console.warn('Failed to preserve hash on login page:', e);
+		}
+	}
+
 	const signInMethod = await $.get('/signInMethod');
 
 	if (signInMethod.okta) {
@@ -71,10 +80,22 @@ $(async () => {
 							redirectString = redirectCookie
 						}
 					}
-		
+
+					// Retrieve preserved hash from sessionStorage
+					let preservedHash = '';
+					try {
+						const hash = sessionStorage.getItem('robomerge_auth_hash');
+						if (hash) {
+							preservedHash = hash;
+							sessionStorage.removeItem('robomerge_auth_hash');
+						}
+					} catch (e) {
+						console.warn('Failed to retrieve preserved hash:', e);
+					}
+
 					document.cookie = `auth=${token}; secure=true; path=/;`;
 					document.cookie = 'redirect_to=;';
-					window.location = window.origin + redirectString;
+					window.location = window.origin + redirectString + preservedHash;
 				})
 				.fail((xhr, status, error) => {
 					$('html').show();
@@ -130,11 +151,23 @@ $(async () => {
 				redirectString = decodeURIComponent(urlParams.get("redirect"))
 			}
 
+			// Retrieve preserved hash from sessionStorage
+			let preservedHash = '';
+			try {
+				const hash = sessionStorage.getItem('robomerge_auth_hash');
+				if (hash) {
+					preservedHash = hash;
+					sessionStorage.removeItem('robomerge_auth_hash');
+				}
+			} catch (e) {
+				console.warn('Failed to retrieve preserved hash:', e);
+			}
+
 			document.cookie = `auth=${token}; secure=true`;
 			document.cookie = 'redirect_to=;';
 
 			// @todo remove redirect_to cookie on server!
-			window.location = window.origin + redirectString
+			window.location = window.origin + redirectString + preservedHash
 
 		}, (xhr, status, error) => {
 			// bad luck
