Changelist links are now buttons with options

diff --git a/Engine/Extras/RoboMerge/v3/public/css/main.css b/Engine/Extras/RoboMerge/v3/public/css/main.css
--- a/Engine/Extras/RoboMerge/v3/public/css/main.css
+++ b/Engine/Extras/RoboMerge/v3/public/css/main.css
@@ -178,11 +178,89 @@ body.stpatricks {
 	box-shadow: inset  0 0 0 0.2rem rgb(192,192,192,0.5); /* silver */
 }
 
-/* Override obnoxious default color for selecting item */ 
+/* Override obnoxious default color for selecting item */
 .dropdown-item.active, .dropdown-item:active {
 	background-color: cadetblue;
 }
 
+/* CL Menu Button Styling */
+.cl-menu-button {
+	display: inline-block !important;
+	vertical-align: baseline;
+	position: relative;
+}
+
+.cl-menu-button .btn {
+	font-family: monospace;
+	padding: .15rem .35rem;
+	font-size: .875rem;
+	line-height: 1.5;
+	white-space: nowrap;
+}
+
+.cl-menu-button .dropdown-menu {
+	min-width: 12rem;
+	font-size: .9rem;
+	display: none;
+	position: absolute;
+	z-index: 1000;
+}
+
+.cl-menu-button .dropdown-menu.show {
+	display: block;
+}
+
+.cl-menu-button .dropdown-item i {
+	width: 1.25em;
+	margin-right: .5em;
+	text-align: center;
+}
+
+/* Soften CL button appearance in conflict blocks */
+.info-block.conflict .cl-menu-button .btn {
+	background-color: rgba(255, 255, 255, 0.5);
+	border-color: rgba(169, 169, 169, 0.4);
+	box-shadow: 0 1px 2px rgba(0, 0, 0, 0.1);
+}
+
+.info-block.conflict .cl-menu-button .btn:hover {
+	background-color: rgba(255, 255, 255, 0.7);
+	border-color: rgba(169, 169, 169, 0.5);
+}
+
+/* Ensure h6 tags containing CL buttons stay inline, but content after wraps to new line */
+.info-block h6 {
+	display: inline-block;
+	margin: 0;
+	width: 100%;
+}
+
+/* Green gate CL button styling - maintain appearance of original last-good-cl */
+.lastchangecell .cl-menu-button.gate-cl {
+	display: block !important;
+	text-align: right;
+	vertical-align: unset;
+}
+
+.cl-menu-button .btn.last-good-cl {
+	font-size: x-small;
+	color: green;
+	background-color: transparent;
+	border: none;
+	padding: 0;
+	margin-right: 0.5em;
+	font-family: inherit;
+	line-height: normal;
+	display: inline-block;
+}
+
+.cl-menu-button .btn.last-good-cl:hover {
+	color: darkgreen;
+	background-color: transparent;
+	border: none;
+	text-decoration: underline;
+}
+
 #signed-in-user {
 	display: inline-flex;
 	align-items: center;
diff --git a/Engine/Extras/RoboMerge/v3/public/css/table.css b/Engine/Extras/RoboMerge/v3/public/css/table.css
--- a/Engine/Extras/RoboMerge/v3/public/css/table.css
+++ b/Engine/Extras/RoboMerge/v3/public/css/table.css
@@ -93,7 +93,7 @@
     flex-direction: column;
     font-weight: 550;
     width: 100%;
-    overflow: auto;
+    overflow: visible; /* Changed from auto to allow CL menu dropdowns to extend beyond cell */
   }
   
   /* Status */
@@ -114,6 +114,7 @@
   /* Last Changelist */
   .lastchangecell {
     flex: 0 0 8em;
+    overflow: visible !important; /* Allow CL menu dropdowns to extend beyond cell */
   }
 
   /* Cell formatting */
@@ -200,10 +201,13 @@
     line-height: normal;
     border: 2px solid gainsboro;
     font-size: 0.8rem;
+    overflow: visible;
+    position: relative;
   }
 
   .namecell .info-block p {
     margin-top: .5em;
+    overflow: visible;
   }
 
   .info-block.pause {
diff --git a/Engine/Extras/RoboMerge/v3/public/js/boilerplate.js b/Engine/Extras/RoboMerge/v3/public/js/boilerplate.js
--- a/Engine/Extras/RoboMerge/v3/public/js/boilerplate.js
+++ b/Engine/Extras/RoboMerge/v3/public/js/boilerplate.js
@@ -114,6 +114,114 @@ function makeSwarmFileLink(depotPath, swarmURL, alias) {
 	}
 }
 
+function makeClMenuButton(cl, swarmURL, options = {}) {
+	// Handle non-numeric CLs - display as bold text without menu
+	if (typeof(cl) !== "number") {
+		return `<span style="font-weight:bolder">${cl}</span>`
+	}
+
+	// Extract options
+	const alias = options.alias || cl
+	const botname = options.botname || ''
+	const nodename = options.nodename || ''
+	const edgename = options.edgename || ''
+	const allowSetLastCL = options.allowSetLastCL || false
+	const allowSetGateCL = options.allowSetGateCL || false
+	const customClass = options.customClass || ''
+	const prefix = options.prefix || ''
+	const tooltip = options.tooltip || ''
+	const jobLink = options.jobLink || ''
+
+	// If no swarmURL, return plain text (no menu functionality)
+	if (!swarmURL) {
+		return `<span>${alias}</span>`
+	}
+
+	// Build data attributes for button
+	let dataAttrs = `data-cl="${cl}" data-swarm-url="${swarmURL}" data-botname="${botname}" data-nodename="${nodename}" data-edgename="${edgename}"`
+
+	if (jobLink) {
+		dataAttrs += ` data-job-link="${jobLink}"`
+	}
+
+	// Store operation function and args if provided (for Set Last CL or Set Gate CL functionality)
+	// Note: These will be retrieved from data attributes in event handlers
+	if ((allowSetLastCL || allowSetGateCL) && options.operationFunction && options.operationArgs) {
+		// Store a reference ID that can be used to retrieve the function from a global map
+		const opId = `cl-op-${cl}-${Date.now()}`
+		if (!window.clMenuOperations) {
+			window.clMenuOperations = {}
+		}
+		window.clMenuOperations[opId] = {
+			operationFunction: options.operationFunction,
+			operationArgs: options.operationArgs
+		}
+		dataAttrs += ` data-operation-id="${opId}"`
+	}
+
+	// Build menu items
+	let menuItems = ''
+
+	// If there's a job link, add it as the first option
+	if (jobLink) {
+		menuItems += `<a class="dropdown-item" href="${jobLink}" target="_blank">
+			<i class="fas fa-external-link-alt"></i> Open CIS Job
+		</a>`
+		menuItems += `<div class="dropdown-divider"></div>`
+	}
+
+	// Copy CL to Clipboard
+	menuItems += `<a class="dropdown-item" href="#" data-action="copy-cl">
+		<i class="fas fa-copy"></i> Copy CL to Clipboard
+	</a>`
+
+	// Open CL in Swarm
+	menuItems += `<a class="dropdown-item" href="${swarmURL}/changes/${cl}" target="_blank">
+		<i class="fas fa-external-link-alt"></i> Open CL in Swarm
+	</a>`
+
+	// Track Changes
+	menuItems += `<a class="dropdown-item" href="#" data-action="track-changes">
+		<i class="fas fa-search"></i> Track Changes
+	</a>`
+
+	// Conditional: Set Last CL option
+	if (allowSetLastCL) {
+		menuItems += `<div class="dropdown-divider"></div>`
+		menuItems += `<a class="dropdown-item" href="#" data-action="set-last-cl">
+			<i class="fas fa-pencil-alt"></i> Set Last CL...
+		</a>`
+	}
+
+	// Conditional: Set Gate CL option
+	if (allowSetGateCL) {
+		menuItems += `<div class="dropdown-divider"></div>`
+		menuItems += `<a class="dropdown-item" href="#" data-action="set-gate-cl">
+			<i class="fas fa-pencil-alt"></i> Set Gate CL...
+		</a>`
+	}
+
+	// Build complete HTML structure
+	const buttonClasses = customClass ? `btn btn-sm btn-secondary-alt ${customClass}` : 'btn btn-sm btn-secondary-alt'
+	const wrapperClasses = customClass === 'last-good-cl' ? 'btn-group cl-menu-button gate-cl' : 'btn-group cl-menu-button'
+	const tooltipAttr = tooltip ? `title="${tooltip}"` : ''
+
+	return `<div class="${wrapperClasses}">
+		<button class="${buttonClasses}"
+				type="button"
+				data-toggle="dropdown"
+				aria-haspopup="true"
+				aria-expanded="false"
+				${tooltipAttr}
+				${dataAttrs}>
+			${prefix}${alias}
+		</button>
+		<div class="dropdown-menu">
+			${menuItems}
+		</div>
+	</div>`
+}
+
 function debug(message) {
 	log('DEBUG: ' + message)
 }
@@ -1105,7 +1213,7 @@ function createPauseDivs(data, conflict) {
 
 		divs.push($('<div class="info-block conflict">')
 			.append(
-				$('<p>').html('<h6>Blocked on change ' + makeClLink(changelist, data.swarmURL) + `</h6> ${info}`)
+				$('<p>').html('<h6>Blocked on change ' + makeClMenuButton(changelist, data.swarmURL) + `</h6> ${info}`)
 			)
 		)
 	}
@@ -1718,12 +1826,19 @@ function renderActionsCell_Edge(actionCell, nodeData, edgeData, conflict=null) {
 function renderLastChangeCell_Common(botname, nodename, lastCl, swarmURL, operationFunction, operationArgs, catchupText=null, edgename=null) {
 	const lastChangeCell = $('<div>').addClass('lastchangecell')
 
-	const swarmLink = $(makeClLink(lastCl, swarmURL)).appendTo(lastChangeCell)
+	const swarmLink = $(makeClMenuButton(lastCl, swarmURL, {
+		botname: botname,
+		nodename: nodename,
+		edgename: edgename,
+		allowSetLastCL: true,
+		operationFunction: operationFunction,
+		operationArgs: operationArgs
+	})).appendTo(lastChangeCell)
 	if (catchupText) {
 		lastChangeCell.append($('<div>').addClass('catchup').text(catchupText))
 	}
 
-	// On shift+click, we can set the CL instead
+	// On shift+click, we can set the CL instead (keeping as fallback during transition)
 	swarmLink.click(function(evt) {
 		if (evt.shiftKey)
 		{
@@ -1787,40 +1902,20 @@ function postRenderLastChangeCell_Edge(botname, lastChangeCell, edgeData, operat
 		if (edgeData.headCL) {
 			tooltip += ` (head changelist ${edgeData.headCL})`
 		}
-		
-		let goodCL = edgeData.lastGoodCLJobLink ? $(`<a href="${edgeData.lastGoodCLJobLink}">`).prop('target', '_blank') : $('<div>');
-		goodCL.html('\u{2713} ' + edgeData.lastGoodCL).addClass('last-good-cl').prop('title', tooltip).appendTo(lastChangeCell)
 
-		if (!edgeData.lastGoodCLJobLink)
-		{
-			// On shift+click, we can set the CL instead
-			goodCL.click(function(evt) {
-				if (evt.shiftKey)
-				{
-					let data = promptFor({
-						cl: {prompt: 'Enter CL', default: edgeData.lastGoodCL},
-					})
-					if (data) {
-						data.reason = "manually set through Robomerge homepage"
-
-						operationFunction(...operationArgs, "/set_gate_cl?" + toQuery(data), function(success) {
-							if (success) {
-								updateBranchList(botname)
-								displaySuccessfulMessage(`Successfully set gate for ${edgeData.displayName} to changelist ${data.cl}`)
-							} else {
-								displayErrorMessage(`Error setting gate for ${edgeData.displayName} to changelist ${data.cl}, please check logs.`)
-							}
-						})
-						
-					}
-					if (evt.preventDefault) {
-						evt.preventDefault()
-					}
-					return false
-				}
-				return true
-			})
-		}
+		// Use CL menu button with custom styling to maintain green gate appearance
+		$(makeClMenuButton(edgeData.lastGoodCL, edgeData.swarmURL, {
+			botname: botname,
+			nodename: edgeData.display_name,
+			edgename: edgeData.displayName,
+			allowSetGateCL: !edgeData.lastGoodCLJobLink,
+			operationFunction: operationFunction,
+			operationArgs: operationArgs,
+			customClass: 'last-good-cl',
+			prefix: '\u{2713} ',
+			tooltip: tooltip,
+			jobLink: edgeData.lastGoodCLJobLink || ''
+		})).appendTo(lastChangeCell)
 	}
 }
 
@@ -2009,3 +2104,111 @@ if (!onLoginPage) {
 	// Preserve hash before potential auth redirect
 	preserveHashForAuth();
 }
+
+/*********************************
+ * CL MENU BUTTON EVENT HANDLERS *
+ *********************************/
+
+// Global event handlers for CL menu button actions using jQuery delegation
+$(document).on('click', '.cl-menu-button .dropdown-item[data-action="copy-cl"]', function(e) {
+	e.preventDefault();
+	const button = $(this).closest('.btn-group').find('button');
+	const cl = button.data('cl');
+
+	if (navigator.clipboard && navigator.clipboard.writeText) {
+		navigator.clipboard.writeText(cl.toString()).catch(() => {
+			// Fallback: show CL in alert
+			alert(`Failed to copy. CL number: ${cl}`);
+		});
+	} else {
+		// Fallback for browsers without clipboard API
+		alert(`CL number: ${cl}`);
+	}
+});
+
+$(document).on('click', '.cl-menu-button .dropdown-item[data-action="track-changes"]', function(e) {
+	e.preventDefault();
+	const button = $(this).closest('.btn-group').find('button');
+	const cl = button.data('cl');
+
+	// Track changes page URL
+	let trackUrl = `/trackchange/${cl}`;
+
+	// Open in new tab
+	window.open(trackUrl, '_blank');
+});
+
+$(document).on('click', '.cl-menu-button .dropdown-item[data-action="set-last-cl"]', function(e) {
+	e.preventDefault();
+	const button = $(this).closest('.btn-group').find('button');
+	const cl = button.data('cl');
+	const botname = button.data('botname');
+	const nodename = button.data('nodename');
+	const edgename = button.data('edgename');
+	const opId = button.data('operation-id');
+
+	// Retrieve the operation function and args from global map
+	if (!window.clMenuOperations || !window.clMenuOperations[opId]) {
+		displayErrorMessage('Set Last CL operation not available');
+		return;
+	}
+
+	const { operationFunction, operationArgs } = window.clMenuOperations[opId];
+
+	// Use existing promptFor dialog
+	let data = promptFor({
+		cl: {prompt: 'Enter CL', default: cl}
+	});
+
+	if (data) {
+		data.reason = "manually set through Robomerge homepage";
+
+		// Call the operation function with the arguments
+		operationFunction(...operationArgs, "/set_last_cl?" + toQuery(data), function(success) {
+			if (success) {
+				updateBranchList(botname);
+				displaySuccessfulMessage(`Successfully set ${edgename || nodename} to changelist ${data.cl}`);
+			} else {
+				displayErrorMessage(`Error setting ${edgename || nodename} to changelist ${data.cl}, please check logs.`);
+				$('#helpButton').trigger('click');
+			}
+		});
+	}
+});
+
+$(document).on('click', '.cl-menu-button .dropdown-item[data-action="set-gate-cl"]', function(e) {
+	e.preventDefault();
+	const button = $(this).closest('.btn-group').find('button');
+	const cl = button.data('cl');
+	const botname = button.data('botname');
+	const edgename = button.data('edgename');
+	const opId = button.data('operation-id');
+
+	// Retrieve the operation function and args from global map
+	if (!window.clMenuOperations || !window.clMenuOperations[opId]) {
+		displayErrorMessage('Set Gate CL operation not available');
+		return;
+	}
+
+	const { operationFunction, operationArgs } = window.clMenuOperations[opId];
+
+	// Use existing promptFor dialog
+	let data = promptFor({
+		cl: {prompt: 'Enter CL', default: cl}
+	});
+
+	if (data) {
+		data.reason = "manually set through Robomerge homepage";
+
+		// Call the operation function with the arguments
+		operationFunction(...operationArgs, "/set_gate_cl?" + toQuery(data), function(success) {
+			if (success) {
+				updateBranchList(botname);
+				displaySuccessfulMessage(`Successfully set gate for ${edgename} to changelist ${data.cl}`);
+			} else {
+				displayErrorMessage(`Error setting gate for ${edgename} to changelist ${data.cl}, please check logs.`);
+				$('#helpButton').trigger('click');
+			}
+		});
+	}
+});
diff --git a/Engine/Extras/RoboMerge/v3/public/js/stomp.js b/Engine/Extras/RoboMerge/v3/public/js/stomp.js
--- a/Engine/Extras/RoboMerge/v3/public/js/stomp.js
+++ b/Engine/Extras/RoboMerge/v3/public/js/stomp.js
@@ -82,7 +82,7 @@ function stompVerify() {
                 $(this).text(targetBranchName)
             })
             $('.changelist').each(function() {
-                $(this).text(requestedBranchCl)
+                $(this).html(makeClMenuButton(requestedBranchCl, requestedNode.swarmURL))
             })
 
             // Data all verified, time to perform the node operation.
@@ -214,9 +214,9 @@ function transitionDisplayStompResults() {
 function visualizeChangelist(changelist, changelistData, swarmURL) {
     const changelistDiv = $(`<div id="${changelist}" class="stomp-visual">`)
     if (robomergeUser && changelistData.author.toLowerCase() === robomergeUser.userName.toLowerCase()) {
-        changelistDiv.append($('<h3>').html(`${makeClLink(changelist, swarmURL)} by <strong>you</strong>`))
+        changelistDiv.append($('<h3>').html(`${makeClMenuButton(changelist, swarmURL)} by <strong>you</strong>`))
     } else {
-        changelistDiv.append($('<h3>').html(`${makeClLink(changelist, swarmURL)} by ${changelistData.author}`))
+        changelistDiv.append($('<h3>').html(`${makeClMenuButton(changelist, swarmURL)} by ${changelistData.author}`))
     }
     changelistDiv.append($('<hr style="margin-top:0em; border-width: 2px;">'))
 
@@ -248,10 +248,10 @@ function visualizeChangelist(changelist, changelistData, swarmURL) {
 }
 
 function displayLimitedInfoModeWarning(requestedBranchCl, stompJson) {
-    const linkToSwarm = makeClLink(requestedBranchCl, stompJson.swarmURL, `CL ${requestedBranchCl}`)
+    const linkToSwarm = makeClMenuButton(requestedBranchCl, stompJson.swarmURL, { alias: `${requestedBranchCl}` })
 
     $('#limitedInfoModeWarning').removeClass('initiallyHidden')
-    $('#limitedInfoClMessage').html(`${linkToSwarm} exceeds the analysis limit.`)
+    $('#limitedInfoClMessage').html(`CL ${linkToSwarm} exceeds the analysis limit.`)
 }
 
 function visualizeStompVerification(requestedBranchCl, stompJson) {
@@ -324,7 +324,7 @@ function visualizeStompVerification(requestedBranchCl, stompJson) {
         }
     }
 
-    const linkToSwarm = makeClLink(requestedBranchCl, stompJson.swarmURL, `CL ${requestedBranchCl}`)
+    const linkToSwarm = makeClMenuButton(requestedBranchCl, stompJson.swarmURL, { alias: `${requestedBranchCl}` })
 
     if (encounteredStompedRevisionsCalculationIssues) {
         $('#stompedRevisionsIssuesDiv').removeClass('initiallyHidden')
@@ -334,7 +334,7 @@ function visualizeStompVerification(requestedBranchCl, stompJson) {
     if (!stompJson.remainingAllBinary) {
         $('#verificationNotAllBinary').removeClass('initiallyHidden')
         $('#resultVisualization').append(
-            $('<h5 class="centered-text">').html(`Robomerge cannot proceed with the stomp request using ${linkToSwarm}. The following changelists contain unstompable changes due to the presence of text files:`)
+            $('<h5 class="centered-text">').html(`Robomerge cannot proceed with the stomp request using CL ${linkToSwarm}. The following changelists contain unstompable changes due to the presence of text files:`)
         )
     } 
 
@@ -344,19 +344,30 @@ function visualizeStompVerification(requestedBranchCl, stompJson) {
             $('#mixedMergeWarning').removeClass('initiallyHidden')
         }
 
-        // Display valid request
-        $('#resultVisualization').append(
-            $('<h5 class="centered-text">').html(
-                `The following changelists will be affected after stomping changes with files from ${linkToSwarm}:`
+        // Check if we have any changelists to display
+        const sortedKeys = Object.keys(clDict).sort()
+
+        if (sortedKeys.length > 0) {
+            // Display valid request with changelists
+            $('#resultVisualization').append(
+                $('<h5 class="centered-text">').html(
+                    `The following changelists will be affected after stomping changes with files from CL ${linkToSwarm}:`
+                )
             )
-        )
-    }
 
-    // Now display each changelist
-    const sortedKeys = Object.keys(clDict).sort()
-    sortedKeys.forEach(function(key) {
-        $('#resultVisualization').append(visualizeChangelist(key, clDict[key], stompJson.swarmURL))
-    })
+            // Display each changelist
+            sortedKeys.forEach(function(key) {
+                $('#resultVisualization').append(visualizeChangelist(key, clDict[key], stompJson.swarmURL))
+            })
+        } else {
+            // No changelists to display
+            $('#resultVisualization').append(
+                $('<h5 class="centered-text">').html(
+                    `Unable to calculate which, if any, changelists will be affected by stomping changes with files from CL ${linkToSwarm}`
+                )
+            )
+        }
+    }
 }
 
 function stompFailure(message) {
diff --git a/Engine/Extras/RoboMerge/v3/public/js/unlock.js b/Engine/Extras/RoboMerge/v3/public/js/unlock.js
--- a/Engine/Extras/RoboMerge/v3/public/js/unlock.js
+++ b/Engine/Extras/RoboMerge/v3/public/js/unlock.js
@@ -77,7 +77,7 @@ function unlockVerify() {
                 $(this).text(targetBranchName)
             })
             $('.changelist').each(function() {
-                $(this).text(requestedBranchCl)
+                $(this).html(makeClMenuButton(requestedBranchCl, requestedNode.swarmURL))
             })
 
             // Data all verified, time to perform the node operation.
