Don't delete and recreate the branchmap workspace in devmode

diff --git a/Engine/Extras/RoboMerge/v3/src/robo/robo.ts b/Engine/Extras/RoboMerge/v3/src/robo/robo.ts
--- a/Engine/Extras/RoboMerge/v3/src/robo/robo.ts
+++ b/Engine/Extras/RoboMerge/v3/src/robo/robo.ts
@@ -454,10 +454,14 @@ async function init(logger: ContextualLogger) {
 		// Ensure we have a workspace for branch specs and that it points to the correct root path
 		const workspace = await p4.find_workspace_by_name(args.branchSpecsWorkspace, {includeUnloaded: true})
 		if (workspace.length === 1) {
-			await p4.deleteWorkspace(args.branchSpecsWorkspace)
+			if (!args.devMode) {
+				await p4.deleteWorkspace(args.branchSpecsWorkspace)
+			}
 		}
-		if ((await lookupStream(args.branchSpecsRootPath)).length > 0) {
-			await p4.newBranchSpecWorkspace(autoUpdaterConfig.workspace, args.branchSpecsRootPath)
+		if (!args.devMode || workspace.length === 0) {
+			if ((await lookupStream(args.branchSpecsRootPath)).length > 0) {
+				await p4.newBranchSpecWorkspace(autoUpdaterConfig.workspace, args.branchSpecsRootPath)
+			}
 		}
 
 		// make sure we've got the latest branch specs
