Remove Set Gate option when the CL is set as an absolute number
Add a link to the swarm history on the gate files

diff --git a/Engine/Extras/RoboMerge/v3/public/js/boilerplate.js b/Engine/Extras/RoboMerge/v3/public/js/boilerplate.js
--- a/Engine/Extras/RoboMerge/v3/public/js/boilerplate.js
+++ b/Engine/Extras/RoboMerge/v3/public/js/boilerplate.js
@@ -131,6 +131,7 @@ function makeClMenuButton(cl, swarmURL, options = {}) {
 	const prefix = options.prefix || ''
 	const tooltip = options.tooltip || ''
 	const jobLink = options.jobLink || ''
+	const fileHistoryLink = options.fileHistoryLink || ''
 
 	// If no swarmURL, return plain text (no menu functionality)
 	if (!swarmURL) {
@@ -185,6 +186,13 @@ function makeClMenuButton(cl, swarmURL, options = {}) {
 		<i class="fas fa-search"></i> Track Changes
 	</a>`
 
+	// If there's a file history link (for gate CLs from file paths), add it
+	if (fileHistoryLink) {
+		menuItems += `<a class="dropdown-item" href="${fileHistoryLink}" target="_blank">
+			<i class="fas fa-history"></i> View Gate File History
+		</a>`
+	}
+
 	// Conditional: Set Last CL option
 	if (allowSetLastCL) {
 		menuItems += `<div class="dropdown-divider"></div>`
@@ -1903,18 +1911,32 @@ function postRenderLastChangeCell_Edge(botname, lastChangeCell, edgeData, operat
 			tooltip += ` (head changelist ${edgeData.headCL})`
 		}
 
+		// Determine if gate is settable (only if lastGoodCLPath is a string, not a number)
+		const isGateSettable = typeof edgeData.lastGoodCLPath === 'string'
+
+		// Build file history link if lastGoodCLPath is a string (file path)
+		let fileHistoryLink = ''
+		if (typeof edgeData.lastGoodCLPath === 'string') {
+			// Remove leading '//' from depot path for Swarm URL, same as makeSwarmFileLink
+			const pathForSwarm = edgeData.lastGoodCLPath.startsWith('//')
+				? edgeData.lastGoodCLPath.slice(2)
+				: edgeData.lastGoodCLPath
+			fileHistoryLink = `${edgeData.swarmURL}/files/${encodeURI(pathForSwarm)}#commits`
+		}
+
 		// Use CL menu button with custom styling to maintain green gate appearance
 		$(makeClMenuButton(edgeData.lastGoodCL, edgeData.swarmURL, {
 			botname: botname,
 			nodename: edgeData.display_name,
 			edgename: edgeData.displayName,
-			allowSetGateCL: !edgeData.lastGoodCLJobLink,
+			allowSetGateCL: isGateSettable && !edgeData.lastGoodCLJobLink,
 			operationFunction: operationFunction,
 			operationArgs: operationArgs,
 			customClass: 'last-good-cl',
 			prefix: '\u{2713} ',
 			tooltip: tooltip,
-			jobLink: edgeData.lastGoodCLJobLink || ''
+			jobLink: edgeData.lastGoodCLJobLink || '',
+			fileHistoryLink: fileHistoryLink
 		})).appendTo(lastChangeCell)
 	}
 }
diff --git a/Engine/Extras/RoboMerge/v3/src/robo/gate.ts b/Engine/Extras/RoboMerge/v3/src/robo/gate.ts
--- a/Engine/Extras/RoboMerge/v3/src/robo/gate.ts
+++ b/Engine/Extras/RoboMerge/v3/src/robo/gate.ts
@@ -411,6 +411,7 @@ export class Gate {
 		const mostRecentGate = this.getMostRecentGate()
 		if (mostRecentGate) {
 			outStatus.lastGoodCL = mostRecentGate.cl
+			outStatus.lastGoodCLPath = this.context.options.lastGoodCLPath
 			outStatus.lastGoodCLJobLink = mostRecentGate.link
 			if (mostRecentGate.timestamp) {
 				outStatus.lastGoodCLDate = new Date(mostRecentGate.timestamp)
diff --git a/Engine/Extras/RoboMerge/v3/src/robo/status-types.ts b/Engine/Extras/RoboMerge/v3/src/robo/status-types.ts
--- a/Engine/Extras/RoboMerge/v3/src/robo/status-types.ts
+++ b/Engine/Extras/RoboMerge/v3/src/robo/status-types.ts
@@ -118,6 +118,7 @@ export type EdgeStatusFields = BotStatusFields & {
 	excludeDescriptions: string[]
 
 	lastGoodCL?: number
+	lastGoodCLPath?: string | number
 	lastGoodCLJobLink?: string
 	lastGoodCLDate?: Date
 	waitingForCISLink?: string
