Improved modal prompt boxes with ability to provide optional help text

diff --git a/Engine/Extras/RoboMerge/v3/public/js/boilerplate.js b/Engine/Extras/RoboMerge/v3/public/js/boilerplate.js
--- a/Engine/Extras/RoboMerge/v3/public/js/boilerplate.js
+++ b/Engine/Extras/RoboMerge/v3/public/js/boilerplate.js
@@ -59,28 +59,100 @@ function createErrorMessage(xhr, error) {
 		return `HTTP ${xhr.status}: ${error}`;
 }
 
-function promptFor(map) {
-	var result = {};
-	var success = true;
-	$.each(Object.keys(map), function(_, key) {
-		if (success) {
-			var p = map[key];
-			var dflt = "";
-			if (typeof(p) === "object") {
-				dflt = p.default || "";
-				p = p.prompt || key;
-			}
+function promptForAsync(map, helpText, modalSize = '') {
+	return new Promise((resolve, reject) => {
+		const keys = Object.keys(map);
+		if (keys.length !== 1) {
+			reject(new Error('promptForAsync only supports single-field prompts'));
+			return;
+		}
+
+		const key = keys[0];
+		const p = map[key];
+		let promptText = key;
+		let defaultValue = "";
+
+		if (typeof(p) === "object") {
+			defaultValue = p.default || "";
+			promptText = p.prompt || key;
+		} else {
+			promptText = p;
+		}
+
+		// Create modal HTML
+		const modalId = 'promptModal' + Date.now();
+		const modalSizeClass = modalSize; // e.g., 'modal-lg' for large, '' for default
+		const modalHtml = `
+			<div class="modal fade" id="${modalId}" tabindex="-1" role="dialog" data-backdrop="static" data-keyboard="false">
+				<div class="modal-dialog ${modalSizeClass}" role="document">
+					<div class="modal-content">
+						<div class="modal-header">
+							<h5 class="modal-title">${promptText}</h5>
+							<button type="button" class="close" data-dismiss="modal" aria-label="Close">
+								<span aria-hidden="true">&times;</span>
+							</button>
+						</div>
+						<div class="modal-body">
+							<input type="text" class="form-control" id="${modalId}_input" value="${defaultValue}">
+							${helpText ? `<small class="form-text text-muted" style="white-space: pre-line; margin-top: 10px; font-family: monospace;">${helpText}</small>` : ''}
+						</div>
+						<div class="modal-footer">
+							<button type="button" class="btn btn-secondary" id="${modalId}_cancel">Cancel</button>
+							<button type="button" class="btn btn-primary" id="${modalId}_ok">OK</button>
+						</div>
+					</div>
+				</div>
+			</div>
+		`;
+
+		// Append modal to body
+		$('body').append(modalHtml);
+		const $modal = $('#' + modalId);
+		const $input = $('#' + modalId + '_input');
+
+		// Show modal
+		$modal.modal('show');
+
+		// Focus input after modal is shown
+		$modal.on('shown.bs.modal', function() {
+			$input.focus();
+		});
+
+		// Handle OK button
+		$('#' + modalId + '_ok').on('click', function() {
+			const value = $input.val();
+			const result = {};
+			result[key] = value;
+			$modal.modal('hide');
+			resolve(result);
+		});
+
+		// Handle Cancel button
+		$('#' + modalId + '_cancel').on('click', function() {
+			$modal.modal('hide');
+			resolve(null);
+		});
 
-			var data = prompt(p, dflt);
-			if (data === null) {
-				success = false;
+		// Handle Enter key in input
+		$input.on('keypress', function(e) {
+			if (e.which === 13) { // Enter key
+				e.preventDefault();
+				$('#' + modalId + '_ok').click();
 			}
-			else {
-				result[key] = data;
+		});
+
+		// Handle Escape key
+		$modal.on('keydown', function(e) {
+			if (e.which === 27) { // Escape key
+				$('#' + modalId + '_cancel').click();
 			}
-		}
+		});
+
+		// Clean up modal after it's hidden
+		$modal.on('hidden.bs.modal', function() {
+			$modal.remove();
+		});
 	});
-	return success ? result : null;
 }
 
 // Add way to prompt the user on leave/reload before the perform an action
@@ -659,24 +731,29 @@ function generateRobomergeFooter() {
 
 	let trackChangeButton = $('<button id="trackChangeButton">')
 	trackChangeButton.addClass("btn btn-sm btn-outline-dark")
-	trackChangeButton.click(function() { 
-		let data = promptFor({
+	trackChangeButton.click(function() {
+		promptForAsync({
 			cl: `Enter the CL to track:`
-		})
-		if (data) {
-			// Ensure they entered a CL
-			if (!isNaN(parseInt(data.cl))) {
-				window.open(`/trackchange/${data.cl}`, '_blank') 
-			}
-			else {
-				const splitCL = data.cl.split(" ")
-				if (splitCL.length != 2 || isNaN(splitCL[1])) {
-					displayErrorMessage("Please provide a valid changelist number to track.")
-					return
+		}).then(data => {
+			if (data) {
+				// Ensure they entered a CL
+				if (!isNaN(parseInt(data.cl))) {
+					window.open(`/trackchange/${data.cl}`, '_blank')
+				}
+				else {
+					const splitCL = data.cl.split(" ")
+					if (splitCL.length != 2 || isNaN(splitCL[1])) {
+						displayErrorMessage("Please provide a valid changelist number to track.")
+						return
+					}
+					window.open(`/trackchange/${splitCL[0]}/${splitCL[1]}`)
 				}
-				window.open(`/trackchange/${splitCL[0]}/${splitCL[1]}`)
 			}
-		}
+		}).catch(err => {
+			if (err) {
+				console.error('Error in track change prompt:', err);
+			}
+		});
 	})
 	trackChangeButton.text("Track Change")
 	trackChangeButton.hide()
@@ -1662,27 +1739,32 @@ function renderActionsCell_Common(actionCell, data, operationFunction, operation
 	if (!data.is_paused) {
 		// Pause Option
 		const pauseOption = createActionOption(`Pause ${dataDisplayName}`, function() {
-			let pauseReasonPrompt = promptFor({
-				msg: `Why are you pausing ${dataDisplayName}?`}
-			);
-			if (pauseReasonPrompt) {
-				pauseReason = pauseReasonPrompt.msg.toString().trim()
-				// Request a reason.
-				if (pauseReason === "") {
-					displayErrorMessage(`Please provide reason for pausing ${dataFullName}. Cancelling pause request.`)
-					return
-				}
-
-				operationFunction(...operationArgs, "/pause?" + toQuery(pauseReasonPrompt), function(success) {
-					if (success) {
-						displaySuccessfulMessage(`Paused ${dataDisplayName}: "${pauseReasonPrompt.msg}"`)
-					} else {
-						displayErrorMessage(`Error pausing ${dataDisplayName}, please check logs.`)
-						$('#helpButton').trigger('click')
+			promptForAsync({
+				msg: `Why are you pausing ${dataDisplayName}?`
+			}).then(pauseReasonPrompt => {
+				if (pauseReasonPrompt) {
+					pauseReason = pauseReasonPrompt.msg.toString().trim()
+					// Request a reason.
+					if (pauseReason === "") {
+						displayErrorMessage(`Please provide reason for pausing ${dataFullName}. Cancelling pause request.`)
+						return
 					}
-					updateBranchList(botname)
-				})
-			}
+
+					operationFunction(...operationArgs, "/pause?" + toQuery(pauseReasonPrompt), function(success) {
+						if (success) {
+							displaySuccessfulMessage(`Paused ${dataDisplayName}: "${pauseReasonPrompt.msg}"`)
+						} else {
+							displayErrorMessage(`Error pausing ${dataDisplayName}, please check logs.`)
+							$('#helpButton').trigger('click')
+						}
+						updateBranchList(botname)
+					})
+				}
+			}).catch(err => {
+				if (err) {
+					console.error('Error in pause prompt:', err);
+				}
+			});
 		},  `Pause Robomerge monitoring of ${dataDisplayName}`)	
 		dropdownEntries.push(pauseOption)
 	}
@@ -1690,10 +1772,15 @@ function renderActionsCell_Common(actionCell, data, operationFunction, operation
 	if (optEdgeName || window.allowNodeReconsiders) {
 		// Add Reconsider action
 		const reconsiderOption = createActionOption(`Reconsider...`, function() {
-			let data = promptFor({
+			const helpText = "Usage: &lt;cl&gt; [#robomerge &lt;stream&gt; | #robomerge[&lt;bot&gt;] &lt;stream&gt; | ...]<br>"
+							 + "Examples:<br>"
+							 + "&nbsp;&nbsp;12345678<br>"
+							 + "&nbsp;&nbsp;12345678 #robomerge targetStream<br>"
+							 + "&nbsp;&nbsp;12345678 #robomerge[botName] targetStream | #robomerge[additionalBotName] additionalTargetStream";
+
+			promptForAsync({
 				cl: `Enter the CL to reconsider (should be a CL from ${dataDisplayName}):`
-			});
-			if (data) {
+			}, helpText, 'modal-lg').then(data => {
 				// Ensure they entered a CL
 				if (isNaN(parseInt(data.cl))) {
 					displayErrorMessage("Please provide a valid changelist number to reconsider.")
@@ -1709,7 +1796,12 @@ function renderActionsCell_Common(actionCell, data, operationFunction, operation
 					}
 					updateBranchList(botname)
 				});
-			}
+			}).catch(err => {
+				// User cancelled or error occurred
+				if (err) {
+					console.error('Error in reconsider prompt:', err);
+				}
+			});
 		}, `Manually submit a CL to Robomerge to process (should be a CL from ${dataDisplayName})`)
 		dropdownEntries.push(reconsiderOption)
 	}
@@ -1834,7 +1926,7 @@ function renderActionsCell_Edge(actionCell, nodeData, edgeData, conflict=null) {
 function renderLastChangeCell_Common(botname, nodename, lastCl, swarmURL, operationFunction, operationArgs, catchupText=null, edgename=null) {
 	const lastChangeCell = $('<div>').addClass('lastchangecell')
 
-	const swarmLink = $(makeClMenuButton(lastCl, swarmURL, {
+	$(makeClMenuButton(lastCl, swarmURL, {
 		botname: botname,
 		nodename: nodename,
 		edgename: edgename,
@@ -1846,36 +1938,6 @@ function renderLastChangeCell_Common(botname, nodename, lastCl, swarmURL, operat
 		lastChangeCell.append($('<div>').addClass('catchup').text(catchupText))
 	}
 
-	// On shift+click, we can set the CL instead (keeping as fallback during transition)
-	swarmLink.click(function(evt) {
-		if (evt.shiftKey)
-		{
-			let data = promptFor({
-				cl: {prompt: 'Enter CL', default: lastCl},
-			})
-			if (data) {
-				data.reason = "manually set through Robomerge homepage"
-
-				operationFunction(...operationArgs, "/set_last_cl?" + toQuery(data), function(success) {
-					if (success) {
-						updateBranchList(botname)
-						displaySuccessfulMessage(`Successfully set ${edgename || nodename} to changelist ${data.cl}`)
-					} else {
-						displayErrorMessage(`Error setting ${edgename || nodename} to changelist ${data.cl}, please check logs.`)
-						$('#helpButton').trigger('click')
-					}
-				})
-				
-			}
-			if (evt.preventDefault) {
-				evt.preventDefault()
-			}
-			return false
-		}
-		return true
-	})
-	// shift+click end
-
 	return lastChangeCell
 }
 
@@ -2177,25 +2239,29 @@ $(document).on('click', '.cl-menu-button .dropdown-item[data-action="set-last-cl
 
 	const { operationFunction, operationArgs } = window.clMenuOperations[opId];
 
-	// Use existing promptFor dialog
-	let data = promptFor({
+	// Use async modal dialog
+	promptForAsync({
 		cl: {prompt: 'Enter CL', default: cl}
-	});
-
-	if (data) {
-		data.reason = "manually set through Robomerge homepage";
+	}).then(data => {
+		if (data) {
+			data.reason = "manually set through Robomerge homepage";
 
-		// Call the operation function with the arguments
-		operationFunction(...operationArgs, "/set_last_cl?" + toQuery(data), function(success) {
-			if (success) {
-				updateBranchList(botname);
-				displaySuccessfulMessage(`Successfully set ${edgename || nodename} to changelist ${data.cl}`);
-			} else {
-				displayErrorMessage(`Error setting ${edgename || nodename} to changelist ${data.cl}, please check logs.`);
-				$('#helpButton').trigger('click');
-			}
-		});
-	}
+			// Call the operation function with the arguments
+			operationFunction(...operationArgs, "/set_last_cl?" + toQuery(data), function(success) {
+				if (success) {
+					updateBranchList(botname);
+					displaySuccessfulMessage(`Successfully set ${edgename || nodename} to changelist ${data.cl}`);
+				} else {
+					displayErrorMessage(`Error setting ${edgename || nodename} to changelist ${data.cl}, please check logs.`);
+					$('#helpButton').trigger('click');
+				}
+			});
+		}
+	}).catch(err => {
+		if (err) {
+			console.error('Error in set last CL prompt:', err);
+		}
+	});
 });
 
 $(document).on('click', '.cl-menu-button .dropdown-item[data-action="set-gate-cl"]', function(e) {
@@ -2214,23 +2280,27 @@ $(document).on('click', '.cl-menu-button .dropdown-item[data-action="set-gate-cl
 
 	const { operationFunction, operationArgs } = window.clMenuOperations[opId];
 
-	// Use existing promptFor dialog
-	let data = promptFor({
+	// Use async modal dialog
+	promptForAsync({
 		cl: {prompt: 'Enter CL', default: cl}
-	});
-
-	if (data) {
-		data.reason = "manually set through Robomerge homepage";
+	}).then(data => {
+		if (data) {
+			data.reason = "manually set through Robomerge homepage";
 
-		// Call the operation function with the arguments
-		operationFunction(...operationArgs, "/set_gate_cl?" + toQuery(data), function(success) {
-			if (success) {
-				updateBranchList(botname);
-				displaySuccessfulMessage(`Successfully set gate for ${edgename} to changelist ${data.cl}`);
-			} else {
-				displayErrorMessage(`Error setting gate for ${edgename} to changelist ${data.cl}, please check logs.`);
-				$('#helpButton').trigger('click');
-			}
-		});
-	}
+			// Call the operation function with the arguments
+			operationFunction(...operationArgs, "/set_gate_cl?" + toQuery(data), function(success) {
+				if (success) {
+					updateBranchList(botname);
+					displaySuccessfulMessage(`Successfully set gate for ${edgename} to changelist ${data.cl}`);
+				} else {
+					displayErrorMessage(`Error setting gate for ${edgename} to changelist ${data.cl}, please check logs.`);
+					$('#helpButton').trigger('click');
+				}
+			});
+		}
+	}).catch(err => {
+		if (err) {
+			console.error('Error in set gate CL prompt:', err);
+		}
+	});
 });
