Can now reconsider from the trackchanges page

diff --git a/Engine/Extras/RoboMerge/v3/public/js/boilerplate.js b/Engine/Extras/RoboMerge/v3/public/js/boilerplate.js
--- a/Engine/Extras/RoboMerge/v3/public/js/boilerplate.js
+++ b/Engine/Extras/RoboMerge/v3/public/js/boilerplate.js
@@ -59,102 +59,6 @@ function createErrorMessage(xhr, error) {
 		return `HTTP ${xhr.status}: ${error}`;
 }
 
-function promptForAsync(map, helpText, modalSize = '') {
-	return new Promise((resolve, reject) => {
-		const keys = Object.keys(map);
-		if (keys.length !== 1) {
-			reject(new Error('promptForAsync only supports single-field prompts'));
-			return;
-		}
-
-		const key = keys[0];
-		const p = map[key];
-		let promptText = key;
-		let defaultValue = "";
-
-		if (typeof(p) === "object") {
-			defaultValue = p.default || "";
-			promptText = p.prompt || key;
-		} else {
-			promptText = p;
-		}
-
-		// Create modal HTML
-		const modalId = 'promptModal' + Date.now();
-		const modalSizeClass = modalSize; // e.g., 'modal-lg' for large, '' for default
-		const modalHtml = `
-			<div class="modal fade" id="${modalId}" tabindex="-1" role="dialog" data-backdrop="static" data-keyboard="false">
-				<div class="modal-dialog ${modalSizeClass}" role="document">
-					<div class="modal-content">
-						<div class="modal-header">
-							<h5 class="modal-title">${promptText}</h5>
-							<button type="button" class="close" data-dismiss="modal" aria-label="Close">
-								<span aria-hidden="true">&times;</span>
-							</button>
-						</div>
-						<div class="modal-body">
-							<input type="text" class="form-control" id="${modalId}_input" value="${defaultValue}">
-							${helpText ? `<small class="form-text text-muted" style="white-space: pre-line; margin-top: 10px; font-family: monospace;">${helpText}</small>` : ''}
-						</div>
-						<div class="modal-footer">
-							<button type="button" class="btn btn-secondary" id="${modalId}_cancel">Cancel</button>
-							<button type="button" class="btn btn-primary" id="${modalId}_ok">OK</button>
-						</div>
-					</div>
-				</div>
-			</div>
-		`;
-
-		// Append modal to body
-		$('body').append(modalHtml);
-		const $modal = $('#' + modalId);
-		const $input = $('#' + modalId + '_input');
-
-		// Show modal
-		$modal.modal('show');
-
-		// Focus input after modal is shown
-		$modal.on('shown.bs.modal', function() {
-			$input.focus();
-		});
-
-		// Handle OK button
-		$('#' + modalId + '_ok').on('click', function() {
-			const value = $input.val();
-			const result = {};
-			result[key] = value;
-			$modal.modal('hide');
-			resolve(result);
-		});
-
-		// Handle Cancel button
-		$('#' + modalId + '_cancel').on('click', function() {
-			$modal.modal('hide');
-			resolve(null);
-		});
-
-		// Handle Enter key in input
-		$input.on('keypress', function(e) {
-			if (e.which === 13) { // Enter key
-				e.preventDefault();
-				$('#' + modalId + '_ok').click();
-			}
-		});
-
-		// Handle Escape key
-		$modal.on('keydown', function(e) {
-			if (e.which === 27) { // Escape key
-				$('#' + modalId + '_cancel').click();
-			}
-		});
-
-		// Clean up modal after it's hidden
-		$modal.on('hidden.bs.modal', function() {
-			$modal.remove();
-		});
-	});
-}
-
 // Add way to prompt the user on leave/reload before the perform an action
 function addOnBeforeUnload(message) {
 	$( window ).on("beforeunload", function() {
@@ -544,47 +448,6 @@ function handleUserPermissions(user) {
 	}
 }
 
-// Adds a message along with bootstrap alert styles
-// https://www.w3schools.com/bootstrap/bootstrap_ref_js_alert.asp
-function displayMessage(message, alertStyle, closable = true) {
-	// Focus user attention
-	window.scrollTo(0, 0);
-
-	let messageDiv = $(`<div class="alert ${alertStyle} fade in show" role="alert">`)
-
-	// Add closable button
-	if (closable) {
-		messageDiv.addClass("alert-dismissible")
-
-		// Display 'X' button
-		let button = $('<button type="button" class="close" data-dismiss="alert" aria-label="Close">')
-		button.html('&times;')
-		messageDiv.append(button)
-
-		// Fade out after 10 seconds
-		setTimeout(function() {
-			messageDiv.slideUp().alert('close')
-		}, 10000)
-	}
-
-	messageDiv.append(message)
-
-	$('#status_message').append(messageDiv)
-}
-// Helper functions wrapping displayMessage()
-function displaySuccessfulMessage(message, closable = true) {
-	displayMessage(`<strong>Success!</strong> ${message}`, "alert-success", closable)
-}
-function displayInfoMessage(message, closable = true) {
-	displayMessage(`${message}`, "alert-info", closable)
-}
-function displayWarningMessage(message, closable = true) {
-	displayMessage(`<strong>Warning:</strong> ${message}`, "alert-warning", closable)
-}
-function displayErrorMessage(message, closable = true) {
-	displayMessage(`<strong>Error:</strong> ${message}`, "alert-danger", closable)
-}
-
 // This function takes string arrays of key names (requiredKeys, optionalKeys) and attempts to get them from the query parameters.
 // If successful, returns an object of { key: value, ... }
 // If failure, displays an error message and returns null
@@ -1772,15 +1635,7 @@ function renderActionsCell_Common(actionCell, data, operationFunction, operation
 	if (optEdgeName || window.allowNodeReconsiders) {
 		// Add Reconsider action
 		const reconsiderOption = createActionOption(`Reconsider...`, function() {
-			const helpText = "Usage: &lt;cl&gt; [#robomerge &lt;stream&gt; | #robomerge[&lt;bot&gt;] &lt;stream&gt; | ...]<br>"
-							 + "Examples:<br>"
-							 + "&nbsp;&nbsp;12345678<br>"
-							 + "&nbsp;&nbsp;12345678 #robomerge targetStream<br>"
-							 + "&nbsp;&nbsp;12345678 #robomerge[botName] targetStream | #robomerge[additionalBotName] additionalTargetStream";
-
-			promptForAsync({
-				cl: `Enter the CL to reconsider (should be a CL from ${dataDisplayName}):`
-			}, helpText, 'modal-lg').then(data => {
+			showReconsiderDialog(`Enter the CL to reconsider (should be a CL from ${dataDisplayName}):`, '').then(data => {
 				// Ensure they entered a CL
 				if (isNaN(parseInt(data.cl))) {
 					displayErrorMessage("Please provide a valid changelist number to reconsider.")
diff --git a/Engine/Extras/RoboMerge/v3/public/js/trackchange.js b/Engine/Extras/RoboMerge/v3/public/js/trackchange.js
--- a/Engine/Extras/RoboMerge/v3/public/js/trackchange.js
+++ b/Engine/Extras/RoboMerge/v3/public/js/trackchange.js
@@ -1,6 +1,49 @@
 // Copyright Epic Games, Inc. All Rights Reserved.
 "use strict"
 
+// displayMessage, displaySuccessfulMessage, displayErrorMessage, promptForAsync
+// and showReconsiderDialog are provided by ui-helpers.js, loaded before this script.
+
+function doReconsider(bot, node, cl, edgeName, targetLabel) {
+	const title = targetLabel ? `Reconsider CL#${cl} \u2192 ${targetLabel}` : `Reconsider CL#${cl}`
+	showReconsiderDialog(title, String(cl))
+		.then(data => {
+			if (!data) return
+			const base = `/api/op/bot/${encodeURIComponent(bot)}/node/${encodeURIComponent(node)}`
+			const opPath = edgeName
+				? `${base}/edge/${encodeURIComponent(edgeName)}/op/reconsider`
+				: `${base}/op/reconsider`
+			$.post(`${opPath}?${new URLSearchParams({cl: data.cl}).toString()}`)
+				.then(() => displaySuccessfulMessage(`Reconsider of CL#${cl} queued successfully`))
+				.catch(err => displayErrorMessage(`Error reconsidering CL#${cl}: ${err.responseText || err.message}`))
+		})
+}
+
+function formatReconsiderButton(cl, reconsiderInfo) {
+	if (!reconsiderInfo || reconsiderInfo.length === 0) {
+		return '<td></td>'
+	}
+	let html = '<td><div class="btn-group">'
+	html += `<button type="button" class="btn btn-sm btn-outline-secondary dropdown-toggle" data-toggle="dropdown">Reconsider</button>`
+	html += `<div class="dropdown-menu">`
+	for (let i = 0; i < reconsiderInfo.length; i++) {
+		const entry = reconsiderInfo[i]
+		if (i > 0) {
+			html += `<div class="dropdown-divider"></div>`
+		}
+		html += `<h6 class="dropdown-header">${entry.bot}</h6>`
+		for (const target of entry.targets) {
+			const edgeLower = target.edgeName.toLowerCase()
+			const label = edgeLower === 'main' || edgeLower.endsWith('-main')
+				? getStreamGraphName(target.stream)
+				: target.edgeName
+			html += `<a class="dropdown-item" href="#" onclick="doReconsider('${entry.bot}','${entry.node}',${cl},'${target.edgeName}','${label.replace(/'/g, "\\'")}');return false;">${label}</a>`
+		}
+	}
+	html += `</div></div></td>`
+	return html
+}
+
 function getMergeMethodStyle(mergeMethod) {
 	switch(mergeMethod) {
 		case "automerge":
@@ -103,6 +146,7 @@ async function generateChangeList(dataObj) {
 			html += `<td>CL#${cl}</td>`
 		}
 		html += `${formatGitHubCommit(gitHubCommits.get(clKey))}`
+		html += formatReconsiderButton(cl, changeDetails.reconsiderInfo)
 		html += '</tr>'
 	}
 	html += "</tbody></table></div>"
diff --git a/Engine/Extras/RoboMerge/v3/public/js/ui-helpers.js b/Engine/Extras/RoboMerge/v3/public/js/ui-helpers.js
new file mode 100644
--- /dev/null
+++ b/Engine/Extras/RoboMerge/v3/public/js/ui-helpers.js
@@ -0,0 +1,145 @@
+// Copyright Epic Games, Inc. All Rights Reserved.
+// ui-helpers.js -- Shared UI utility functions (modals, alerts) used across multiple pages.
+//                  Must be loaded before boilerplate.js and any page-specific scripts.
+
+// Displays a Bootstrap modal prompting the user for a single text value.
+// map:       { <key>: <promptText> | { prompt: <string>, default: <string> } }
+// helpText:  optional HTML shown below the input
+// modalSize: optional Bootstrap modal size class e.g. 'modal-lg'
+// okLabel:   optional label for the confirm button (default 'OK')
+// Returns a Promise resolving to { <key>: <value> } on confirm, or null on cancel.
+function promptForAsync(map, helpText, modalSize = '', okLabel = 'OK') {
+	return new Promise((resolve, reject) => {
+		const keys = Object.keys(map);
+		if (keys.length !== 1) {
+			reject(new Error('promptForAsync only supports single-field prompts'));
+			return;
+		}
+
+		const key = keys[0];
+		const p = map[key];
+		let promptText = key;
+		let defaultValue = "";
+
+		if (typeof(p) === "object") {
+			defaultValue = p.default || "";
+			promptText = p.prompt || key;
+		} else {
+			promptText = p;
+		}
+
+		const modalId = 'promptModal' + Date.now();
+		const modalHtml = `
+			<div class="modal fade" id="${modalId}" tabindex="-1" role="dialog" data-backdrop="static" data-keyboard="false">
+				<div class="modal-dialog ${modalSize}" role="document">
+					<div class="modal-content">
+						<div class="modal-header">
+							<h5 class="modal-title">${promptText}</h5>
+							<button type="button" class="close" data-dismiss="modal" aria-label="Close">
+								<span aria-hidden="true">&times;</span>
+							</button>
+						</div>
+						<div class="modal-body">
+							<input type="text" class="form-control" id="${modalId}_input" value="${defaultValue}">
+							${helpText ? `<small class="form-text text-muted" style="white-space: pre-line; margin-top: 10px; font-family: monospace;">${helpText}</small>` : ''}
+						</div>
+						<div class="modal-footer">
+							<button type="button" class="btn btn-secondary" id="${modalId}_cancel">Cancel</button>
+							<button type="button" class="btn btn-primary" id="${modalId}_ok">${okLabel}</button>
+						</div>
+					</div>
+				</div>
+			</div>
+		`;
+
+		$('body').append(modalHtml);
+		const $modal = $('#' + modalId);
+		const $input = $('#' + modalId + '_input');
+
+		$modal.modal('show');
+
+		$modal.on('shown.bs.modal', function() {
+			$input.focus();
+		});
+
+		$('#' + modalId + '_ok').on('click', function() {
+			const value = $input.val();
+			const result = {};
+			result[key] = value;
+			$modal.modal('hide');
+			resolve(result);
+		});
+
+		$('#' + modalId + '_cancel').on('click', function() {
+			$modal.modal('hide');
+			resolve(null);
+		});
+
+		$input.on('keypress', function(e) {
+			if (e.which === 13) {
+				e.preventDefault();
+				$('#' + modalId + '_ok').click();
+			}
+		});
+
+		$modal.on('keydown', function(e) {
+			if (e.which === 27) {
+				$('#' + modalId + '_cancel').click();
+			}
+		});
+
+		$modal.on('hidden.bs.modal', function() {
+			$modal.remove();
+		});
+	});
+}
+
+// Opens a Bootstrap modal prompting the user to reconsider a CL.
+// promptText: title shown in the modal header
+// defaultCL:  pre-populated value for the CL input (pass '' for blank)
+// okLabel:    optional label for the confirm button (default 'Reconsider')
+// Returns a Promise resolving to the entered CL string, or null on cancel.
+const RECONSIDER_HELP_TEXT =
+	'Usage: &lt;cl&gt; [#robomerge &lt;stream&gt; | #robomerge[&lt;bot&gt;] &lt;stream&gt; | ...]<br>'
+	+ 'Examples:<br>'
+	+ '&nbsp;&nbsp;12345678<br>'
+	+ '&nbsp;&nbsp;12345678 #robomerge targetStream<br>'
+	+ '&nbsp;&nbsp;12345678 #robomerge[botName] targetStream | #robomerge[additionalBotName] additionalTargetStream'
+
+function showReconsiderDialog(promptText, defaultCL, okLabel = 'Reconsider') {
+	return promptForAsync({cl: {prompt: promptText, default: defaultCL}}, RECONSIDER_HELP_TEXT, 'modal-lg', okLabel)
+}
+
+// Adds a Bootstrap alert message to #status_message.
+// closable: if true (default), adds a dismiss button and auto-fades after 10s.
+function displayMessage(message, alertStyle, closable = true) {
+	window.scrollTo(0, 0);
+
+	let messageDiv = $(`<div class="alert ${alertStyle} fade in show" role="alert">`)
+
+	if (closable) {
+		messageDiv.addClass("alert-dismissible")
+		let button = $('<button type="button" class="close" data-dismiss="alert" aria-label="Close">')
+		button.html('&times;')
+		messageDiv.append(button)
+		setTimeout(function() {
+			messageDiv.slideUp().alert('close')
+		}, 10000)
+	}
+
+	messageDiv.append(message)
+	$('#status_message').append(messageDiv)
+}
+
+function displaySuccessfulMessage(message, closable = true) {
+	displayMessage(`<strong>Success!</strong> ${message}`, "alert-success", closable)
+}
+function displayInfoMessage(message, closable = true) {
+	displayMessage(`${message}`, "alert-info", closable)
+}
+function displayWarningMessage(message, closable = true) {
+	displayMessage(`<strong>Warning:</strong> ${message}`, "alert-warning", closable)
+}
+function displayErrorMessage(message, closable = true) {
+	displayMessage(`<strong>Error:</strong> ${message}`, "alert-danger", closable)
+}
diff --git a/Engine/Extras/RoboMerge/v3/public/trackchange.html b/Engine/Extras/RoboMerge/v3/public/trackchange.html
--- a/Engine/Extras/RoboMerge/v3/public/trackchange.html
+++ b/Engine/Extras/RoboMerge/v3/public/trackchange.html
@@ -7,6 +7,7 @@
 
 	<script src="/js/jquery-3.4.1.min.js"></script>
 	<script src="/js/bootstrap.bundle.min.js"></script>
+	<script src="/js/ui-helpers.js"></script>
 	<script src="/js/viz-standalone.js"></script>
 	<script src="/js/flow.graph.js"></script>
 	<script src="/js/flow.js"></script>
@@ -24,6 +25,7 @@
 </head>
 <body>
 	<header><a href="javascript:window.location.href=window.location.origin;"><img src="/img/logo.png"></a></header>
+	<div id="status_message"></div>
 	<div class="container">
 
 		<div id="loadingDiv" class="card card-single border-dark loading-box">
diff --git a/Engine/Extras/RoboMerge/v3/src/robo/trackchange.ts b/Engine/Extras/RoboMerge/v3/src/robo/trackchange.ts
--- a/Engine/Extras/RoboMerge/v3/src/robo/trackchange.ts
+++ b/Engine/Extras/RoboMerge/v3/src/robo/trackchange.ts
@@ -2,7 +2,7 @@
 
 import { ContextualLogger } from '../common/logger';
 import { DescribeResult, isExecP4Error, PerforceContext } from '../common/perforce';
-import { Node } from './graph';
+import { Graph, Node } from './graph';
 import { NodeBotInterface } from './bot-interfaces';
 import { RoboMerge } from './robo';
 import { Status } from './status';
@@ -47,6 +47,43 @@ interface GatherOptions {
     inSourceChain?: boolean
 }
 
+function isBotVisibleToUser(botName: string, graph: Graph, userTags: Set<string>, logger: ContextualLogger) {
+    const bg = graph.branchGraphAliases.get(botName.toUpperCase())
+    return bg ? Status.includeBranch(bg.config.visibility, userTags, logger) : false
+}
+
+function buildReconsiderInfo(node: Node, graph: Graph, userTags: Set<string>, botFilter: string[], logger: ContextualLogger) {
+    const byBot = new Map<string, {node: string, targets: {edgeName: string, stream: string}[]}>()
+    for (const edge of graph.getEdgesBySource(node)) {
+        const nodeBot = edge.sourceAnnotation as NodeBotInterface
+        const botName = edge.bot as string
+        if (botFilter.length > 0 && !botFilter.includes(botName)) continue
+
+        // Use the authoritative branchGraphAliases registry to check source bot visibility
+        if (!isBotVisibleToUser(botName, graph, userTags, logger)) continue
+
+        // Also check the target stream is visible: it must appear as a source in at least
+        // one visible bot. If the target has no outgoing edges (leaf node) we allow it
+        // through since the source bot visibility check above is sufficient.
+        const targetOutEdges = [...graph.getEdgesBySource(edge.target)]
+        if (targetOutEdges.length > 0) {
+            const targetVisible = targetOutEdges.some(te =>
+                isBotVisibleToUser(te.bot as string, graph, userTags, logger))
+            if (!targetVisible) continue
+        }
+
+        if (!byBot.has(botName)) {
+            byBot.set(botName, {node: nodeBot.branch.name, targets: []})
+        }
+        byBot.get(botName)!.targets.push({
+            edgeName: edge.targetName,
+            stream: edge.target.stream as string
+        })
+    }
+    if (byBot.size === 0) return undefined
+    return [...byBot.entries()].map(([bot, info]) => ({bot, node: info.node, targets: info.targets}))
+}
+
 export async function trackChange(robo: RoboMerge, logger: ContextualLogger, initialCL: TrackedCL, userTags: Set<string>, filters: TrackFilters) {
 
     const streamFilter = (() => {
@@ -390,7 +427,8 @@ export async function trackChange(robo: RoboMerge, logger: ContextualLogger, ini
 
         if (includeInResults) {
             const swarmLink = p4.swarmURL ? `${p4.swarmURL}/changes/${clToConsider.cl}` : undefined
-            data.changes[`${clToConsider.cl}`] = {streamDisplayName, swarmLink, mergeMethod: changeToConsider.mergeMethod, sourceCL: changeToConsider.sourceCL?.cl}
+            const reconsiderInfo = changeToConsider.node ? buildReconsiderInfo(changeToConsider.node, graph, userTags, botFilter, logger) : undefined
+            data.changes[`${clToConsider.cl}`] = {streamDisplayName, swarmLink, mergeMethod: changeToConsider.mergeMethod, sourceCL: changeToConsider.sourceCL?.cl, ...(reconsiderInfo && {reconsiderInfo})}
         }
     }
 
