Fix undefined access when there are not trackchanges results

diff --git a/Engine/Extras/RoboMerge/v3/public/js/trackchange.js b/Engine/Extras/RoboMerge/v3/public/js/trackchange.js
--- a/Engine/Extras/RoboMerge/v3/public/js/trackchange.js
+++ b/Engine/Extras/RoboMerge/v3/public/js/trackchange.js
@@ -219,7 +219,7 @@ async function buildResults(data) {
 		const $errorPanel = $('#error-panel');
 		$('pre', $errorPanel).html('No results found.')
 		if (data.data.swarmURL) {
-			$('.swarmURL').html(`<a href="${data.data.swarmURL}/changes/${data.originalCL.cl}" target="_blank">CL# ${data.originalCL.cl}</a>`)
+			$('.swarmURL').html(`<a href="${data.data.swarmURL}/changes/${data.data.originalCL.cl}" target="_blank">CL# ${data.data.originalCL.cl}</a>`)
 		}
 		$errorPanel.show();
 	}
