Prevent from graph becoming unredable when only a few nodes
Better layout for track changes graph when only a few nodes

diff --git a/Engine/Extras/RoboMerge/v3/public/js/flow.js b/Engine/Extras/RoboMerge/v3/public/js/flow.js
--- a/Engine/Extras/RoboMerge/v3/public/js/flow.js
+++ b/Engine/Extras/RoboMerge/v3/public/js/flow.js
@@ -108,8 +108,10 @@ function showFlowGraph(data, args, addlegend) {
         $('#graph-loading-text').hide();
         const span = $('<span id="svgSpan">').css('margin', 'auto').html(svg);
         const svgEl = $('svg', span).addClass('branch-graph').removeAttr('width');
-        // scale graph to 70% of default size
-        const height = Math.round(parseInt(svgEl.attr('height')) * .7);
+        // scale graph to 70% of default size, but never smaller than MIN_GRAPH_HEIGHT
+        // (single-node graphs render tiny without a floor)
+        const MIN_GRAPH_HEIGHT = 60;
+        const height = Math.max(Math.round(parseInt(svgEl.attr('height')) * .7), MIN_GRAPH_HEIGHT);
         svgEl.attr('height', height + 'px').css('vertical-align', 'top');
         flowGraph.empty();
         flowGraph.append($('<img src="/img/save.png" style="position:absolute" onclick="savePNG()" title="Save Graph as PNG" width="16">'));
diff --git a/Engine/Extras/RoboMerge/v3/public/js/trackchange.js b/Engine/Extras/RoboMerge/v3/public/js/trackchange.js
--- a/Engine/Extras/RoboMerge/v3/public/js/trackchange.js
+++ b/Engine/Extras/RoboMerge/v3/public/js/trackchange.js
@@ -188,20 +188,22 @@ function createGraph(changes) {
 	}
 	lines.push('}')
 
-    const graphContainer = $('<div class="clearfix">');
-    const flowGraph = $('<div class="flow-graph" style="display:inline-block;">').appendTo(graphContainer);
+    const graphContainer = $('<div>').css({display: 'flex', justifyContent: 'center', alignItems: 'flex-start', gap: '40px', width: '100%'});
+    const flowGraph = $('<div class="flow-graph">').css({flex: '0 0 auto', minWidth: '300px', float: 'none', display: 'block'}).appendTo(graphContainer);
     flowGraph.append($('<div>').css('text-align', 'center').text("Building graph..."));
     renderGraph(lines.join('\n'))
         .then(svg => {
         $('#graph-key-template')
             .clone()
             .removeAttr('id')
-            .css('display', 'inline-block')
+            .css({flex: '0 0 auto', display: 'block'})
             .appendTo(graphContainer);
         const span = $('<div style="margin: auto; display: flex; justify-content: center;">').html(svg);
         const svgEl = $('svg', span).addClass('branch-graph').removeAttr('width');
-        // scale graph to 70% of default size
-        const height = Math.round(parseInt(svgEl.attr('height')) * .7);
+        // scale graph to 70% of default size, but never smaller than MIN_GRAPH_HEIGHT
+        // (single-node graphs render tiny without a floor)
+        const MIN_GRAPH_HEIGHT = 60;
+        const height = Math.max(Math.round(parseInt(svgEl.attr('height')) * .7), MIN_GRAPH_HEIGHT);
         svgEl.attr('height', height + 'px').css('vertical-align', 'top');
         flowGraph.empty();
         flowGraph.append(span);
