Fix graph flowing off the left side of the screen

diff --git a/Engine/Extras/RoboMerge/v3/public/js/trackchange.js b/Engine/Extras/RoboMerge/v3/public/js/trackchange.js
--- a/Engine/Extras/RoboMerge/v3/public/js/trackchange.js
+++ b/Engine/Extras/RoboMerge/v3/public/js/trackchange.js
@@ -188,22 +188,21 @@ function createGraph(changes) {
 	}
 	lines.push('}')
 
-    const graphContainer = $('<div>').css({display: 'flex', justifyContent: 'center', alignItems: 'flex-start', gap: '40px', width: '100%'});
-    const flowGraph = $('<div class="flow-graph">').css({flex: '0 0 auto', minWidth: '300px', float: 'none', display: 'block'}).appendTo(graphContainer);
+    const graphContainer = $('<div>').css({textAlign: 'center', width: '100%', overflowX: 'auto'});
+    const flowGraph = $('<div class="flow-graph">').css({display: 'inline-block', float: 'none', verticalAlign: 'top', minWidth: '400px'}).appendTo(graphContainer);
     flowGraph.append($('<div>').css('text-align', 'center').text("Building graph..."));
     renderGraph(lines.join('\n'))
         .then(svg => {
         $('#graph-key-template')
             .clone()
             .removeAttr('id')
-            .css({flex: '0 0 auto', display: 'block'})
+            .css({display: 'inline-block', verticalAlign: 'top', marginLeft: '40px'})
             .appendTo(graphContainer);
         const span = $('<div style="margin: auto; display: flex; justify-content: center;">').html(svg);
-        const svgEl = $('svg', span).addClass('branch-graph').removeAttr('width');
-        // scale graph to 70% of default size, but never smaller than MIN_GRAPH_HEIGHT
-        // (single-node graphs render tiny without a floor)
-        const MIN_GRAPH_HEIGHT = 60;
-        const height = Math.max(Math.round(parseInt(svgEl.attr('height')) * .7), MIN_GRAPH_HEIGHT);
+        const svgEl = $('svg', span).addClass('branch-graph');
+        const naturalHeight = parseInt(svgEl.attr('height'));
+        svgEl.removeAttr('width');
+        const height = Math.max(Math.round(naturalHeight * 0.7), 60);
         svgEl.attr('height', height + 'px').css('vertical-align', 'top');
         flowGraph.empty();
         flowGraph.append(span);
