Update the blockage thread when the blockage status changes
Don't shrink the in memory node when the DB size is exceeded, instead shrink a copy and try to save that keeping the full data in memory and only losing it on a restart that reloads from the DB

diff --git a/Engine/Extras/RoboMerge/v3/src/robo/conflict-interfaces.ts b/Engine/Extras/RoboMerge/v3/src/robo/conflict-interfaces.ts
--- a/Engine/Extras/RoboMerge/v3/src/robo/conflict-interfaces.ts
+++ b/Engine/Extras/RoboMerge/v3/src/robo/conflict-interfaces.ts
@@ -38,6 +38,12 @@ export interface PersistentConflict {
 	// Set when the Acknowledge button is clicked
 	acknowledger?: string
 	acknowledgedAt?: Date
+
+	// Compact fingerprint of failure details for circumstance-change detection on retry.
+	// Content varies by kind — see buildDetailsSnapshot() in conflicts.ts.
+	// May be absent after a service restart if the persisted copy was shrunk for BSON limits.
+	detailsSnapshot?: string[]
+	detailsSnapshotTruncated?: boolean   // true if list exceeded MAX_SNAPSHOT_ENTRIES
 }
 
 export function PersistentConflictToString(conflict: PersistentConflict) {
diff --git a/Engine/Extras/RoboMerge/v3/src/robo/conflicts.ts b/Engine/Extras/RoboMerge/v3/src/robo/conflicts.ts
--- a/Engine/Extras/RoboMerge/v3/src/robo/conflicts.ts
+++ b/Engine/Extras/RoboMerge/v3/src/robo/conflicts.ts
@@ -1,7 +1,7 @@
 // Copyright Epic Games, Inc. All Rights Reserved.
 
 import { ContextualLogger } from 'common/logger';
-import { AlreadyIntegrated, Blockage, Branch, BranchArg, branchesMatch } from './branch-interfaces';
+import { AlreadyIntegrated, Blockage, Branch, BranchArg, branchesMatch, ExclusiveLockInfo } from './branch-interfaces';
 import { ChangeInfo, ForcedCl, resolveBranchArg } from './branch-interfaces';
 import { ConflictStatusFields } from './status-types';
 import { PersistentConflict, PersistentConflictToString, Resolution } from './conflict-interfaces';
@@ -26,6 +26,27 @@ function conflictMatchesByChangeCl(conflict: PersistentConflict, changeCl: numbe
 }
 
 
+const MAX_SNAPSHOT_ENTRIES = 500
+
+function buildDetailsSnapshot(blockage: Blockage): { snapshot: string[], truncated: boolean } {
+	let entries: string[]
+	if (blockage.failure.kind === 'Merge conflict' && blockage.failure.details) {
+		entries = blockage.failure.details.split('\n').map(s => s.trim()).filter(Boolean).sort()
+	}
+	else if (blockage.failure.kind === 'Exclusive check-out') {
+		const lockInfo = blockage.failure.additionalInfo as ExclusiveLockInfo | undefined
+		entries = (lockInfo?.exclusiveFiles ?? [])
+			.map(f => `${f.user.toLowerCase()}:${f.depotPath}`)
+			.sort()
+	}
+	else {
+		// Syntax error and all other kinds: snapshot the description text for change detection
+		entries = blockage.failure.description ? [blockage.failure.description] : []
+	}
+	const truncated = entries.length > MAX_SNAPSHOT_ENTRIES
+	return { snapshot: entries.slice(0, MAX_SNAPSHOT_ENTRIES), truncated }
+}
+
 function makePersistentConflict(blockage: Blockage): PersistentConflict {
 	const sourceChange = blockage.change
 	const result: PersistentConflict = {
@@ -49,6 +70,14 @@ function makePersistentConflict(blockage: Blockage): PersistentConflict {
 		}
 	}
 
+	const { snapshot, truncated } = buildDetailsSnapshot(blockage)
+	if (snapshot.length > 0) {
+		result.detailsSnapshot = snapshot
+		if (truncated) {
+			result.detailsSnapshotTruncated = true
+		}
+	}
+
 	return result
 }
 
@@ -189,6 +218,15 @@ export class Conflicts {
 			const conflict = this.conflicts[index]
 			if (conflictMatches(conflict, replacement.sourceCl, replacement.targetBranchName)) {
 
+				// Detect if failure circumstances changed before replacing
+				const oldConflict = conflict
+				const kindChanged = oldConflict.kind !== replacement.kind
+				// Only compare when both snapshots are present — avoids false positives after
+				// a service restart or BSON-shrink that cleared the persisted snapshot
+				const snapshotChanged = oldConflict.detailsSnapshot && replacement.detailsSnapshot
+					? JSON.stringify(oldConflict.detailsSnapshot) !== JSON.stringify(replacement.detailsSnapshot)
+					: false
+
 				// maintain certain fields case-by-case
 				const acknowledger = conflict.acknowledger
 				if (acknowledger) {
@@ -203,6 +241,10 @@ export class Conflicts {
 				this.persist()
 
 				this.eventTriggers.reportBlockage(blockage, false)
+
+				if (kindChanged || snapshotChanged) {
+					this.eventTriggers.reportBlockageStatusChanged({ previousConflict: oldConflict, newBlockage: blockage })
+				}
 				return
 			}
 		}
diff --git a/Engine/Extras/RoboMerge/v3/src/robo/events.ts b/Engine/Extras/RoboMerge/v3/src/robo/events.ts
--- a/Engine/Extras/RoboMerge/v3/src/robo/events.ts
+++ b/Engine/Extras/RoboMerge/v3/src/robo/events.ts
@@ -20,6 +20,12 @@ type OnConflictStatus = (arg: boolean) => void
 type OnBeginIntegratingToGate = (arg: BeginIntegratingToGateEvent) => void
 type OnEndIntegratingToGate = (arg: EndIntegratingToGateEvent) => void
 
+export interface BlockageStatusChangedEvent {
+	previousConflict: PersistentConflict
+	newBlockage: Blockage
+}
+type OnBlockageStatusChanged = (arg: BlockageStatusChangedEvent) => void
+
 export interface BotEventHandler {
 	onChangeParsed?: OnChangeParsed
 	onBlockage?: OnBlockage
@@ -35,6 +41,7 @@ export interface BotEventHandler {
 
 	onBeginIntegratingToGate?: OnBeginIntegratingToGate
 	onEndIntegratingToGate?: OnEndIntegratingToGate
+	onBlockageStatusChanged?: OnBlockageStatusChanged
 }
 
 /** Bindings for bot events */
@@ -107,6 +114,10 @@ export class BotEvents {
 		this.endIntegratingToGateListeners.push(listener)
 	}
 
+	onBlockageStatusChanged(listener: OnBlockageStatusChanged) {
+		this.blockageStatusChangedListeners.push(listener)
+	}
+
 	// register an object to handle some or all of the above
 	registerHandler(handler: BotEventHandler) {
 		const proto = handler.constructor.prototype
@@ -127,6 +138,7 @@ export class BotEvents {
 				case 'onConflictStatus': this.conflictStatusListeners.push(arg => prop.call(handler, arg)); break
 				case 'onBeginIntegratingToGate': this.beginIntegratingToGateListeners.push(arg => prop.call(handler, arg)); break
 				case 'onEndIntegratingToGate': this.endIntegratingToGateListeners.push(arg => prop.call(handler, arg)); break
+			case 'onBlockageStatusChanged': this.blockageStatusChangedListeners.push(arg => prop.call(handler, arg)); break
 			}
 		}
 	}
@@ -145,6 +157,7 @@ export class BotEvents {
 	protected readonly conflictStatusListeners: OnConflictStatus[] = []
 	protected readonly beginIntegratingToGateListeners: OnBeginIntegratingToGate[] = []
 	protected readonly endIntegratingToGateListeners: OnEndIntegratingToGate[] = []
+	protected readonly blockageStatusChangedListeners: OnBlockageStatusChanged[] = []
 }
 
 /** API for firing bot events */
@@ -232,4 +245,10 @@ export class BotEventTriggers extends BotEvents {
 			listener(arg)
 		}
 	}
+
+	reportBlockageStatusChanged(arg: BlockageStatusChangedEvent) {
+		for (const listener of this.blockageStatusChangedListeners) {
+			listener(arg)
+		}
+	}
 }
diff --git a/Engine/Extras/RoboMerge/v3/src/robo/notifications.ts b/Engine/Extras/RoboMerge/v3/src/robo/notifications.ts
--- a/Engine/Extras/RoboMerge/v3/src/robo/notifications.ts
+++ b/Engine/Extras/RoboMerge/v3/src/robo/notifications.ts
@@ -8,7 +8,7 @@ import { ContextualLogger } from '../common/logger';
 import { ExclusiveFileDetails } from '../common/perforce';
 import { Blockage, Branch, BranchArg, EditableBranch, ExclusiveLockInfo, ForcedCl, MergeAction, NagBlockageEvent, NodeOpUrlGenerator, PauseEvent, resolveBranchArg } from './branch-interfaces';
 import { PersistentConflict, Resolution } from './conflict-interfaces';
-import { BotEventHandler, BotEvents } from './events';
+import { BlockageStatusChangedEvent, BotEventHandler, BotEvents } from './events';
 import { NodeBot } from './nodebot';
 import { PersistedSlackMessage, PersistedSlackMessages } from './persistedmessages';
 import { RoboArgs } from './roboargs';
@@ -582,50 +582,6 @@ export class BotNotifications implements BotEventHandler {
 
 		if (targetBranch) {
 			messagesToPost.push({ message })
-
-			if (isNew && blockage.failure.details) {
-				let file: SlackFile = {
-					content: blockage.failure.details,
-					channels: message.channel,
-					filename: "conflictdetails.txt"
-				}
-				messagesToPost.push({file})
-			}
-
-			if (isNew && blockage.failure.kind === 'Exclusive check-out') {
-				const exclusiveLockInfo = blockage.failure.additionalInfo as ExclusiveLockInfo
-
-				let authorDict = new Map<string, ExclusiveFileDetails[]>()
-				for (const file of exclusiveLockInfo.exclusiveFiles) {
-					setDefault(authorDict, file.user.toLowerCase(), []).push(file)
-				}
-
-				for (const exclusiveLockUser of exclusiveLockInfo.exclusiveLockUsers) {
-					let file: SlackFile = { 
-						content: authorDict.get(exclusiveLockUser.user)!.map(ef => ef.depotPath).join("\n"),
-						channels: message.channel,
-						filename: "exclusivelockedfiles.txt"
-					}
-					if (exclusiveLockUser.user.length > 0) {
-						let exclusiveLockSlackUser
-						const exclusiveLockUserEmail = await exclusiveLockUser.userEmail
-						if (exclusiveLockUserEmail) {
-							exclusiveLockSlackUser = await this.slackMessages.getSlackUser(exclusiveLockUserEmail)
-							if (exclusiveLockSlackUser) {
-								exclusiveLockSlackUser = `<@${exclusiveLockSlackUser}> `
-							}
-							usersToInvite.add(exclusiveLockUserEmail)
-						}
-						if (!exclusiveLockSlackUser) {
-							exclusiveLockSlackUser = `@${exclusiveLockUser.user} `
-						}
-						file.initial_comment = `${exclusiveLockSlackUser} please unlock the files blocking robomerge or work with ${blockage.owner} to resolve the conflict. Hit retry on the blocked stream once the files are unlocked.`
-					} else {
-						file.initial_comment = "The following locked files did not have their owner determined and as such those owners may not have been notified"
-					}
-					messagesToPost.push({file})
-				}
-			}
 		}
 		else {
 			messagesToPost.push({ message })
@@ -647,10 +603,7 @@ export class BotNotifications implements BotEventHandler {
 				this.slackMessages.addUserToChannel(user, channel)
 			}
 			for (const messageToPost of messagesToPost) {
-				if (messageToPost.file) {
-					await this.slackMessages.postFile(changeInfo.cl, branchArg, { ...messageToPost.file, channels: channel})
-				}
-				else if (messageToPost.reply) {
+				if (messageToPost.reply) {
 					await this.slackMessages.postReply(changeInfo.cl, branchArg, { ...messageToPost.message, channel })
 				}
 				else {
@@ -659,6 +612,11 @@ export class BotNotifications implements BotEventHandler {
 			}
 		}
 
+		// For first-time blockages, post supplementary details (conflict files, locked file pings)
+		if (isNew) {
+			await this._postBlockageDetails(blockage, changeInfo.cl, branchArg, channelsToPostTo)
+		}
+
 		// Post message to owner in DM
 		if (!isBotUser && targetBranch && userEmail) {
 			let dm: SlackMessage
@@ -691,6 +649,111 @@ export class BotNotifications implements BotEventHandler {
 		}
 	}
 
+	/** Post supplementary details for a blockage: conflict file list and/or lock holder pings.
+	 *  Called for new blockages (isNew=true) and when circumstances change (onBlockageStatusChanged). */
+	private async _postBlockageDetails(
+		blockage: Blockage,
+		cl: number,
+		branchArg: string,
+		channels: string[]
+	) {
+		if (!this.slackMessages) return
+
+		// Upload conflict details for merge conflicts
+		if (blockage.failure.details) {
+			const file: SlackFile = {
+				content: blockage.failure.details,
+				channels: '',
+				filename: 'conflictdetails.txt'
+			}
+			for (const channel of channels) {
+				await this.slackMessages.postFile(cl, branchArg, { ...file, channels: channel })
+			}
+		}
+
+		// Upload locked files per holder for exclusive check-outs
+		if (blockage.failure.kind === 'Exclusive check-out') {
+			const exclusiveLockInfo = blockage.failure.additionalInfo as ExclusiveLockInfo
+			const authorDict = new Map<string, ExclusiveFileDetails[]>()
+			for (const f of exclusiveLockInfo.exclusiveFiles) {
+				setDefault(authorDict, f.user.toLowerCase(), []).push(f)
+			}
+			for (const exclusiveLockUser of exclusiveLockInfo.exclusiveLockUsers) {
+				let file: SlackFile = {
+					content: authorDict.get(exclusiveLockUser.user)!.map(ef => ef.depotPath).join('\n'),
+					channels: '',
+					filename: 'exclusivelockedfiles.txt'
+				}
+				let exclusiveLockSlackUser: string | undefined
+				if (exclusiveLockUser.user.length > 0) {
+					const exclusiveLockUserEmail = await exclusiveLockUser.userEmail
+					if (exclusiveLockUserEmail) {
+						const slackId = await this.slackMessages.getSlackUser(exclusiveLockUserEmail)
+						if (slackId) {
+							exclusiveLockSlackUser = `<@${slackId}> `
+						}
+						for (const channel of channels) {
+							this.slackMessages.addUserToChannel(exclusiveLockUserEmail, channel)
+						}
+					}
+					if (!exclusiveLockSlackUser) {
+						exclusiveLockSlackUser = `@${exclusiveLockUser.user} `
+					}
+					file.initial_comment = `${exclusiveLockSlackUser} please unlock the files blocking robomerge or work with ${blockage.owner} to resolve the conflict. Hit retry on the blocked stream once the files are unlocked.`
+				} else {
+					file.initial_comment = 'The following locked files did not have their owner determined and as such those owners may not have been notified'
+				}
+				for (const channel of channels) {
+					await this.slackMessages.postFile(cl, branchArg, { ...file, channels: channel })
+				}
+			}
+		}
+	}
+
+	/** Fired when the failure kind or file set changes on a retry of the same blocked CL.
+	 *  Posts a thread reply noting the transition, then re-posts the full details for the
+	 *  new failure kind as if it were first occurrence. */
+	async onBlockageStatusChanged(evt: BlockageStatusChangedEvent) {
+		if (!this.slackMessages) return
+		const { previousConflict, newBlockage } = evt
+		const changeInfo = newBlockage.change
+		if (changeInfo.userRequest) return
+
+		const sourceBranch = changeInfo.branch
+		const targetBranch = newBlockage.action?.branch
+		const branchArg = targetBranch ? targetBranch.name : newBlockage.failure.kind
+		const channels = this.getChannelsToPostTo(sourceBranch, targetBranch?.upperName)
+		const cl = changeInfo.cl
+
+		// Build context line for the thread reply
+		const kindChanged = previousConflict.kind !== newBlockage.failure.kind
+		let contextLine: string
+		if (kindChanged) {
+			contextLine = `*Status update:* blockage type changed from *${previousConflict.kind}* to *${newBlockage.failure.kind}*`
+		} else {
+			contextLine = `*Status update:* ${newBlockage.failure.kind} circumstances have changed`
+		}
+
+		// For non-target-branch failures (syntax errors etc.), append the description to the reply
+		if (!targetBranch && newBlockage.failure.description) {
+			contextLine += `\n\n${newBlockage.failure.description}`
+		}
+
+		// Post a brief thread reply noting the change
+		const replyMessage: SlackMessage = {
+			text: contextLine,
+			style: SlackMessageStyles.WARNING,
+			channel: '',
+			mrkdwn: true
+		}
+		for (const channel of channels) {
+			await this.slackMessages.postReply(cl, branchArg, { ...replyMessage, channel })
+		}
+
+		// Post the full supplementary details for the new failure kind
+		await this._postBlockageDetails(newBlockage, cl, branchArg, channels)
+	}
+
 	onBlockageAcknowledged(info: PersistentConflict) {
 		if (this.slackMessages) {
 			const targetKey = info.targetBranchName || info.kind
diff --git a/Engine/Extras/RoboMerge/v3/src/robo/settings.ts b/Engine/Extras/RoboMerge/v3/src/robo/settings.ts
--- a/Engine/Extras/RoboMerge/v3/src/robo/settings.ts
+++ b/Engine/Extras/RoboMerge/v3/src/robo/settings.ts
@@ -248,13 +248,13 @@ export class Settings {
 		if (this.enableSave) {
 			if (RoboArgs.useMongo) {
 				if (opts?.name) {
-					const saveToDB = () => {
+					const saveToDB = (saveObject: any) => {
 						this.savingState.set(opts.name!, 'Saving')
-						this.object[opts.name!].version = this.object.version
-						BotModel.collection.updateOne({botname:this.botname,nodename:opts.name},{$set:{botname:this.botname,nodename:opts.name,details:this.object[opts.name!]}},{upsert:true})
+						saveObject.version = this.object.version
+						BotModel.collection.updateOne({botname:this.botname,nodename:opts.name},{$set:{botname:this.botname,nodename:opts.name,details:saveObject}},{upsert:true})
 						.then(() => {
 							if (this.savingState.get(opts.name!) == 'NeedsSave') {
-								saveToDB()
+								saveToDB(this.object[opts.name!])
 							}
 							else {
 								this.savingState.delete(opts.name!)
@@ -263,15 +263,19 @@ export class Settings {
 						.catch(() => {
 							if (this.savingState.get(opts.name!) === 'NeedsSave') {
 								this.settingsLogger.warn(`Unable to save ${this.botname}:${opts.name}. Attempting resave before applying mitigation as settings have changed.`)
-								saveToDB()
-							} 
+								saveToDB(this.object[opts.name!])
+							}
 							else {
-								const size = Buffer.byteLength(JSON.stringify(this.object[opts.name!]));
-								this._shrinkNode(opts.name!)
-								const shrunkSize = Buffer.byteLength(JSON.stringify(this.object[opts.name!]));
+								// Deep-clone the live node data so _shrinkNodeObject does NOT
+								// modify the in-memory state — only the persisted copy is trimmed.
+								// Compare against saveObject size: if shrinking produces the same size as what we just tried, there is no point retrying.
+								const shrunkCopy = JSON.parse(JSON.stringify(this.object[opts.name!]))
+								const size = Buffer.byteLength(JSON.stringify(saveObject))
+								this._shrinkNodeObject(shrunkCopy)
+								const shrunkSize = Buffer.byteLength(JSON.stringify(shrunkCopy))
 								if (shrunkSize < size) {
 									this.settingsLogger.warn(`Unable to save ${this.botname}:${opts.name}. Size: ${size}. Attempting resave with shrunk size ${shrunkSize}.`)
-									saveToDB()
+									saveToDB(shrunkCopy)
 								}
 								else {
 									const errorMsg = `Unable to save ${this.botname}:${opts.name}. Size: ${size}.`
@@ -283,7 +287,7 @@ export class Settings {
 						})
 					}
 					if (!this.savingState.has(opts.name)) {
-						saveToDB()
+						saveToDB(this.object[opts.name!])
 					}
 					else {
 						this.savingState.set(opts.name, 'NeedsSave')
@@ -348,10 +352,14 @@ export class Settings {
 			})
 	}
 
-	_shrinkNode(nodeName: string, trim?: boolean) {
-		const edges = this.object[nodeName]?.edges || {}
+	/** Shrink a node object in place. Used for migration and as the implementation
+	 *  base for the deep-clone save path. Operates on any plain object — callers
+	 *  are responsible for passing either the live node or a disposable clone. */
+	_shrinkNodeObject(nodeObj: any, trim?: boolean) {
+		// Clear describeResult from edge pause states (large file list)
+		const edges = nodeObj?.edges || {}
 		for (const key of Object.keys(edges)) {
-			const edge = this.object[nodeName].edges[key]
+			const edge = edges[key]
 			const additionalInfo: MergeConflictAdditionalInformation = edge.pause?.blockage?.additionalInfo
 			if (isA_MergeConflictAdditionalInformation(additionalInfo)) {
 				if (additionalInfo?.describeResult) {
@@ -362,6 +370,19 @@ export class Settings {
 				}
 			}
 		}
+		// Clear detailsSnapshot from persisted conflicts (can be large for big CLs)
+		for (const conflict of (nodeObj?.conflicts ?? [])) {
+			delete conflict.detailsSnapshot
+			delete conflict.detailsSnapshotTruncated
+		}
+	}
+
+	/** Shrink the live in-memory node object. Used for migration paths that
+	 *  intentionally want to permanently discard large fields. The MongoDB
+	 *  error-recovery path should use the deep-clone approach instead to
+	 *  preserve in-memory data. */
+	_shrinkNode(nodeName: string, trim?: boolean) {
+		this._shrinkNodeObject(this.object[nodeName], trim)
 	}
 }
 
