Can specify #stomp when reconsidering to do an accept source on conflicting binary files on the immediate next merge. Does not propagate to subsequent merges, does not affect text files.

diff --git a/Engine/Extras/RoboMerge/v3/public/js/ui-helpers.js b/Engine/Extras/RoboMerge/v3/public/js/ui-helpers.js
--- a/Engine/Extras/RoboMerge/v3/public/js/ui-helpers.js
+++ b/Engine/Extras/RoboMerge/v3/public/js/ui-helpers.js
@@ -100,11 +100,14 @@ function promptForAsync(map, helpText, modalSize = '', okLabel = 'OK') {
 // okLabel:    optional label for the confirm button (default 'Reconsider')
 // Returns a Promise resolving to the entered CL string, or null on cancel.
 const RECONSIDER_HELP_TEXT =
-	'Usage: &lt;cl&gt; [#robomerge &lt;stream&gt; | #robomerge[&lt;bot&gt;] &lt;stream&gt; | ...]<br>'
+	'<div style="display:flex">Usage:&nbsp;<span>&lt;cl&gt; [#stomp] [#robomerge &lt;stream&gt; | #robomerge[&lt;bot&gt;] &lt;stream&gt; | ...]<span></div><br>'
 	+ 'Examples:<br>'
-	+ '&nbsp;&nbsp;12345678<br>'
-	+ '&nbsp;&nbsp;12345678 #robomerge targetStream<br>'
-	+ '&nbsp;&nbsp;12345678 #robomerge[botName] targetStream | #robomerge[additionalBotName] additionalTargetStream'
+	+ '<div style="padding-left:2ch">12345678</div>'
+	+ '<div style="padding-left:2ch">12345678 #robomerge targetStream</div>'
+	+ '<div style="display:flex">&nbsp;&nbsp;12345678&nbsp;<span>#robomerge[botName] targetStream | #robomerge[additionalBotName] additionalTargetStream</span></div>'
+	+ '<br>'
+	+ 'Flags:<br>'
+	+ '<div style="display:flex">&nbsp;&nbsp;#stomp&nbsp;&mdash;&nbsp;<span>If a merge conflict occurs, stomp binary file conflicts instead of creating a shelf for manual resolution. Stomp applies only to the initial merge and is not propagated to subsequent merges in the flow.</span></div>'
 
 function showReconsiderDialog(promptText, defaultCL, okLabel = 'Reconsider') {
 	return promptForAsync({cl: {prompt: promptText, default: defaultCL}}, RECONSIDER_HELP_TEXT, 'modal-lg', okLabel)
diff --git a/Engine/Extras/RoboMerge/v3/src/robo/bot-interfaces.ts b/Engine/Extras/RoboMerge/v3/src/robo/bot-interfaces.ts
--- a/Engine/Extras/RoboMerge/v3/src/robo/bot-interfaces.ts
+++ b/Engine/Extras/RoboMerge/v3/src/robo/bot-interfaces.ts
@@ -84,7 +84,7 @@ export interface NodeBotInterface extends NodeBotIPC {
 // slightly sneaky that since NodeBot.reconsider's additional args are optional, it fulfils both of these, i.e.
 // implements the base bot reconsider function, and it also the 'queue change' function.
 
-// ReconsiderArgs are the additional args used by stomp etc. I'm adding commandOverride to them, even thoguh that's a
+// ReconsiderArgs are the additional args used by stomp etc. I'm adding commandOverride to them, even though that's a
 // slightly different pattern (unless it breaks something)
 
 export type ReconsiderArgs = {
diff --git a/Engine/Extras/RoboMerge/v3/src/robo/ipc.ts b/Engine/Extras/RoboMerge/v3/src/robo/ipc.ts
--- a/Engine/Extras/RoboMerge/v3/src/robo/ipc.ts
+++ b/Engine/Extras/RoboMerge/v3/src/robo/ipc.ts
@@ -3,7 +3,7 @@
 import { ContextualLogger, isNpmLogLevel, NpmLogLevelCompare, NpmLogLevelValues } from '../common/logger';
 import { getRunningPerforceCommands, PerforceContext } from '../common/perforce';
 import { Trace } from './graph';
-import { IPCControls, EdgeBotInterface, NodeBotInterface } from './bot-interfaces';
+import { IPCControls, EdgeBotInterface, NodeBotInterface, ReconsiderArgs } from './bot-interfaces';
 import { OperationResult } from './branch-interfaces';
 import { GraphBot } from './graphbot';
 import { RoboMerge } from './robo';
@@ -468,15 +468,43 @@ export class IPC {
 			return OPERATION_SUCCESS
 
 		case 'reconsider':
-			const argMatch = query.cl.match(/(\d+)\s*(.*)/)
+			const argMatch = query.cl.match(/^\s*(\d+)\s*(.*)/)
 			if (argMatch) {
 				cl = parseInt(argMatch[1])
 			}
 			if (isNaN(cl)) {
 				return {statusCode: 400, message: 'Invalid CL parameter: ' + query.cl}
 			}
+
+			// Parse #stomp from the command text. This flag is only valid in the reconsider
+			// box — it is not in FLAGMAP so it cannot appear in a CL description, and it
+			// is stripped here before the text reaches the description parser, so it is
+			// never propagated to child CLs.
+			const segments = (argMatch![2] || '').split('|')
+			const reconsiderArgs: Partial<ReconsiderArgs> = {}
+			const commandSegments: string[] = []
+			for (const segment of segments) {
+				const segTokens = segment.trim().split(/\s+/).filter(Boolean)
+				const stompIdx = segTokens.findIndex(t => t.toLowerCase() === '#stomp')
+				if (stompIdx !== -1) {
+					if (stompIdx > 0 && /^#robomerge/i.test(segTokens[stompIdx - 1])) {
+						return {statusCode: 400, message: '#stomp must not appear as an argument to a #robomerge command'}
+					}
+					segTokens.splice(stompIdx, 1)
+					reconsiderArgs.additionalFlags = ['stomp']
+				}
+				const remaining = segTokens.join(' ')
+				if (remaining) {
+					commandSegments.push(remaining)
+				}
+			}
+			const commandText = commandSegments.join(' | ')
+			if (commandText) {
+				reconsiderArgs.commandOverride = commandText
+			}
+
 			try {
-				generalOpTarget.reconsider(query.who, cl, {commandOverride: argMatch![2]})
+				generalOpTarget.reconsider(query.who, cl, reconsiderArgs)
 			}
 			catch (err) {
 				throw err
