Give a good error message when reconsidering a CL that does not exist in the depot at all

diff --git a/Engine/Extras/RoboMerge/v3/src/robo/nodebot.ts b/Engine/Extras/RoboMerge/v3/src/robo/nodebot.ts
--- a/Engine/Extras/RoboMerge/v3/src/robo/nodebot.ts
+++ b/Engine/Extras/RoboMerge/v3/src/robo/nodebot.ts
@@ -429,20 +429,28 @@ export class NodeBot extends PerforceStatefulBot implements NodeBotInterface {
 		const changeResult = await this._getChange(fromQueue.cl)
 
 		if (typeof changeResult === 'string') {
-			let reconsiderErrorMessage = `While processing manually queued change ${fromQueue.cl} on ${this.fullName}, `
-			reconsiderErrorMessage += specifiedTargetBranch ? `targeted specifically at ${specifiedTargetBranch.name}` : ''
-			reconsiderErrorMessage += `\n\nError while querying P4 for change ${fromQueue.cl}: ${changeResult}`
+			const clUnknown = /Change \d+ unknown/i.test(changeResult)
+
+			let reconsiderSubject: string
+			let reconsiderErrorMessage = `While processing manually queued change ${fromQueue.cl} on ${this.fullName}`
+			reconsiderErrorMessage += specifiedTargetBranch ? `, targeted specifically at ${specifiedTargetBranch.name}` : ''
+
+			if (clUnknown) {
+				reconsiderSubject = `Reconsider Request Failed -- CL Does Not Exist`
+				reconsiderErrorMessage += `\n\nChange ${fromQueue.cl} does not exist in Perforce.`
+			}
+			else {
+				reconsiderSubject = `Reconsider Request Failed -- Error Querying Perforce`
+				reconsiderErrorMessage += `\n\nError while querying P4 for change ${fromQueue.cl}: ${changeResult}`
+			}
 
 			await this.sendErrorNotificationToUser(
 				fromQueue.who,
 				fromQueue.cl,
-				`Reconsider Request Failed -- Error Querying Perforce`,
+				reconsiderSubject,
 				reconsiderErrorMessage
 			)
 
-			// change is invalid
-			// @todo distinguish between Perforce errors and change not existing in branch
-
 			this.nodeBotLogger.error(`${this.fullName} - Error while querying P4 for change ${fromQueue.cl}: ${changeResult}`)
 			// @todo email admin page user when users have to log in
 			// any way to show error without blocking page?
