Improve error message when reconsidering a CL that isn't from the correct stream
Provide links to swarm on various messages from robomerge

diff --git a/Engine/Extras/RoboMerge/v3/src/robo/branch-interfaces.ts b/Engine/Extras/RoboMerge/v3/src/robo/branch-interfaces.ts
--- a/Engine/Extras/RoboMerge/v3/src/robo/branch-interfaces.ts
+++ b/Engine/Extras/RoboMerge/v3/src/robo/branch-interfaces.ts
@@ -229,6 +229,7 @@ export type ConflictKind = 'merge' | 'branch' | 'delete' | 'unknown'
 export interface AlreadyIntegrated {
 	change: ChangeInfo
 	action: MergeAction
+	reason?: 'already_integrated' | 'no_files'
 }
 
 export class ResolveResultAdditionalInfo {
diff --git a/Engine/Extras/RoboMerge/v3/src/robo/edgebot.ts b/Engine/Extras/RoboMerge/v3/src/robo/edgebot.ts
--- a/Engine/Extras/RoboMerge/v3/src/robo/edgebot.ts
+++ b/Engine/Extras/RoboMerge/v3/src/robo/edgebot.ts
@@ -560,8 +560,8 @@ class EdgeBotImpl extends PerforceStatefulBot {
 		return workspaceName
 	}
 
-	private async SkipIntegration(pending: PendingChange, edgeServerAddress: string|undefined, msg: string) {
-		const event: AlreadyIntegrated = {change: pending.change, action: pending.action}
+	private async SkipIntegration(pending: PendingChange, edgeServerAddress: string|undefined, msg: string, reason?: AlreadyIntegrated['reason']) {
+		const event: AlreadyIntegrated = {change: pending.change, action: pending.action, reason}
 		this.sourceNode.onAlreadyIntegrated(event)
 
 		// integration not necessary
@@ -619,13 +619,19 @@ class EdgeBotImpl extends PerforceStatefulBot {
 				return await this._resolveChangelist(pending)
 
 			case 'already_integrated':
-			case 'no_files': {
 				return await this.SkipIntegration(
 					pending,
 					edgeServerAddress,
-					`Change ${to_integrate} was not necessary in ${this.targetBranch.name}`
+					`Change ${to_integrate} was not necessary in ${this.targetBranch.name}`,
+					'already_integrated'
+				)
+			case 'no_files':
+				return await this.SkipIntegration(
+					pending,
+					edgeServerAddress,
+					`Change ${to_integrate} has no files in ${this.sourceBranch.name} that map to ${this.targetBranch.name}`,
+					'no_files'
 				)
-			}
 		}
 
 		const errors = results as string[]
diff --git a/Engine/Extras/RoboMerge/v3/src/robo/nodebot.ts b/Engine/Extras/RoboMerge/v3/src/robo/nodebot.ts
--- a/Engine/Extras/RoboMerge/v3/src/robo/nodebot.ts
+++ b/Engine/Extras/RoboMerge/v3/src/robo/nodebot.ts
@@ -1476,7 +1476,14 @@ export class NodeBot extends PerforceStatefulBot implements NodeBotInterface {
 		this.conflicts.onAlreadyIntegrated(event)
 
 		if (event.change.userRequest && this.slackMessages) {
-			const message = `Change ${event.change.cl} was not necessary in ${event.action.branch.name}`
+			let message: string
+			const swarmLink = `<${this.p4.swarmURL}/changes/${event.change.cl}|Change ${event.change.cl}>`
+			if (event.reason === 'no_files') {
+				message = `${swarmLink} contains no files in ${event.change.branch.name} that map to ${event.action.branch.name}. Verify the changelist is from the correct stream.`
+			}
+			else {
+				message = `${swarmLink} was not necessary in ${event.action.branch.name}`
+			}
 			let dm: SlackMessage = {
 				text: message,
 				style: SlackMessageStyles.WARNING,
@@ -2278,7 +2285,8 @@ export class NodeBot extends PerforceStatefulBot implements NodeBotInterface {
 		if (pending.change.userRequest) {
 			this.nodeBotLogger.info(shortMessage + ` (reconsider triggered by ${pending.change.owner}):\n${failure.description}.`)
 
-			let message = `While reconsidering CL#${pending.change.cl}`
+			const clLink = `<${this.p4.swarmURL}/changes/${pending.change.cl}|Change ${pending.change.cl}>`
+			let message = `While reconsidering ${clLink}`
 
 			if (!failure.details) {
 				message += '\n' + failure.description
