Resolve circumstances where #robomerge tag causes command to silently do nothing
Treat all #robomerge commands in the reconsider box as a default bot command
Report an error to the user if no targets were found even given that

diff --git a/Engine/Extras/RoboMerge/v3/src/robo/nodebot.ts b/Engine/Extras/RoboMerge/v3/src/robo/nodebot.ts
--- a/Engine/Extras/RoboMerge/v3/src/robo/nodebot.ts
+++ b/Engine/Extras/RoboMerge/v3/src/robo/nodebot.ts
@@ -529,7 +529,17 @@ export class NodeBot extends PerforceStatefulBot implements NodeBotInterface {
 		}
 
 		// process just this change
-		await this._processAndMergeCl(new Map<string, EdgeBot>(this.edges), change, true, fromQueue.who, fromQueue.workspace, specifiedTargetBranch)
+		const mergeResult = await this._processAndMergeCl(new Map<string, EdgeBot>(this.edges), change, true, fromQueue.who, fromQueue.workspace, specifiedTargetBranch)
+		if (mergeResult.errors.length > 0 && change.isUserRequest) {
+			await this.sendErrorNotificationToUser(
+				fromQueue.who,
+				fromQueue.cl,
+				`Reconsider Request Failed -- No Matching Targets`,
+				`While processing manually queued change ${fromQueue.cl} on ${this.fullName}: no targets were found to integrate to.\n\n` +
+				`If you specified a #robomerge directive without a bot qualifier, try qualifying it explicitly: ` +
+				`#robomerge[${this.graphBotName}] <target>`
+			)
+		}
 		return
 	}
 
@@ -1949,7 +1959,10 @@ export class NodeBot extends PerforceStatefulBot implements NodeBotInterface {
 		const commandOverride = (change.commandOverride || '').trim()
 		if (commandOverride) {
 
-			const parsedOverride = parseDescription(commandOverride, {lineEnd: '|', ...parseDescCommonParams})
+			// When commandOverride comes from the reconsider box (a user request), the user is
+			// explicitly targeting this bot, so bare #robomerge directives should be accepted
+			// regardless of whether this is the default bot for the source node.
+			const parsedOverride = parseDescription(commandOverride, {lineEnd: '|', ...parseDescCommonParams, isDefaultBot: parseDescCommonParams.isDefaultBot || !!change.isUserRequest})
 			if (change.accumulateCommandOverride) {
 				parsedLines.append(parsedOverride)
 			}
