Start with the text selected when setting gate and last CL to make pasting in a new value easier

diff --git a/Engine/Extras/RoboMerge/v3/public/js/ui-helpers.js b/Engine/Extras/RoboMerge/v3/public/js/ui-helpers.js
--- a/Engine/Extras/RoboMerge/v3/public/js/ui-helpers.js
+++ b/Engine/Extras/RoboMerge/v3/public/js/ui-helpers.js
@@ -59,7 +59,7 @@ function promptForAsync(map, helpText, modalSize = '', okLabel = 'OK') {
 		$modal.modal('show');
 
 		$modal.on('shown.bs.modal', function() {
-			$input.focus();
+			$input.focus().select();
 		});
 
 		$('#' + modalId + '_ok').on('click', function() {
