Submit missing files that load ui-helpers.js

diff --git a/Engine/Extras/RoboMerge/v3/public/acknowledge.html b/Engine/Extras/RoboMerge/v3/public/acknowledge.html
--- a/Engine/Extras/RoboMerge/v3/public/acknowledge.html
+++ b/Engine/Extras/RoboMerge/v3/public/acknowledge.html
@@ -11,7 +11,8 @@
 	<script src="/js/jquery-3.4.1.min.js"></script>
 	<script src="/js/bootstrap.bundle.min.js"></script>
 
-<script src="/js/boilerplate.js"></script>
+<script src="/js/ui-helpers.js"></script>
+    <script src="/js/boilerplate.js"></script>
 <script src="/js/acknowledge.js"></script>
 <script language="javascript">
 
diff --git a/Engine/Extras/RoboMerge/v3/public/createshelf.html b/Engine/Extras/RoboMerge/v3/public/createshelf.html
--- a/Engine/Extras/RoboMerge/v3/public/createshelf.html
+++ b/Engine/Extras/RoboMerge/v3/public/createshelf.html
@@ -11,7 +11,8 @@
 	<script src="/js/jquery-3.4.1.min.js"></script>
 	<script src="/js/bootstrap.bundle.min.js"></script>
 
-<script src="/js/boilerplate.js"></script>
+<script src="/js/ui-helpers.js"></script>
+    <script src="/js/boilerplate.js"></script>
 <script src="/js/createshelf.js"></script>
 <script language="javascript">
 
diff --git a/Engine/Extras/RoboMerge/v3/public/index.html b/Engine/Extras/RoboMerge/v3/public/index.html
--- a/Engine/Extras/RoboMerge/v3/public/index.html
+++ b/Engine/Extras/RoboMerge/v3/public/index.html
@@ -19,7 +19,8 @@
 	<script src="/js/jquery-3.4.1.min.js"></script>
 	<script src="/js/bootstrap.bundle.min.js"></script>
 
-	<script src="/js/boilerplate.js"></script>
+	<script src="/js/ui-helpers.js"></script>
+    <script src="/js/boilerplate.js"></script>
 	<script src="/js/viz-standalone.js"></script>
 	<script src="/js/flow.graph.js"></script>
 	<script src="/js/flow.js"></script>
diff --git a/Engine/Extras/RoboMerge/v3/public/login.html b/Engine/Extras/RoboMerge/v3/public/login.html
--- a/Engine/Extras/RoboMerge/v3/public/login.html
+++ b/Engine/Extras/RoboMerge/v3/public/login.html
@@ -22,7 +22,8 @@
 <script src="/js/jquery-3.4.1.min.js"></script>
 <script src="/js/bootstrap.bundle.min.js"></script>
 
-<script src="/js/boilerplate.js"></script>
+<script src="/js/ui-helpers.js"></script>
+    <script src="/js/boilerplate.js"></script>
 <script src="/js/login.js"></script>
 
 </head>
diff --git a/Engine/Extras/RoboMerge/v3/public/markdown-template.html b/Engine/Extras/RoboMerge/v3/public/markdown-template.html
--- a/Engine/Extras/RoboMerge/v3/public/markdown-template.html
+++ b/Engine/Extras/RoboMerge/v3/public/markdown-template.html
@@ -12,7 +12,8 @@
 	<script src="/js/jquery-3.4.1.min.js"></script>
 	<script src="/js/bootstrap.bundle.min.js"></script>
 
-<script src="/js/boilerplate.js"></script>
+<script src="/js/ui-helpers.js"></script>
+    <script src="/js/boilerplate.js"></script>
 
 </head>
 <body>
diff --git a/Engine/Extras/RoboMerge/v3/public/skip.html b/Engine/Extras/RoboMerge/v3/public/skip.html
--- a/Engine/Extras/RoboMerge/v3/public/skip.html
+++ b/Engine/Extras/RoboMerge/v3/public/skip.html
@@ -11,7 +11,8 @@
 	<script src="/js/jquery-3.4.1.min.js"></script>
 	<script src="/js/bootstrap.bundle.min.js"></script>
 
-<script src="/js/boilerplate.js"></script>
+<script src="/js/ui-helpers.js"></script>
+    <script src="/js/boilerplate.js"></script>
 <script src="/js/skip.js"></script>
 <script language="javascript">
 
diff --git a/Engine/Extras/RoboMerge/v3/public/stomp.html b/Engine/Extras/RoboMerge/v3/public/stomp.html
--- a/Engine/Extras/RoboMerge/v3/public/stomp.html
+++ b/Engine/Extras/RoboMerge/v3/public/stomp.html
@@ -11,7 +11,8 @@
 	<script src="/js/jquery-3.4.1.min.js"></script>
 	<script src="/js/bootstrap.bundle.min.js"></script>
 
-<script src="/js/boilerplate.js"></script>
+<script src="/js/ui-helpers.js"></script>
+    <script src="/js/boilerplate.js"></script>
 <script src="/js/stomp.js"></script>
 <script language="javascript">
 
diff --git a/Engine/Extras/RoboMerge/v3/public/unlock.html b/Engine/Extras/RoboMerge/v3/public/unlock.html
--- a/Engine/Extras/RoboMerge/v3/public/unlock.html
+++ b/Engine/Extras/RoboMerge/v3/public/unlock.html
@@ -11,7 +11,8 @@
 	<script src="/js/jquery-3.4.1.min.js"></script>
 	<script src="/js/bootstrap.bundle.min.js"></script>
 
-<script src="/js/boilerplate.js"></script>
+<script src="/js/ui-helpers.js"></script>
+    <script src="/js/boilerplate.js"></script>
 <script src="/js/unlock.js"></script>
 <script language="javascript">
 
