Fix gate getting stuck paused when edge last CL changed or restarted with a new gate CL

diff --git a/Engine/Extras/RoboMerge/v3/src/robo/edgebot.ts b/Engine/Extras/RoboMerge/v3/src/robo/edgebot.ts
--- a/Engine/Extras/RoboMerge/v3/src/robo/edgebot.ts
+++ b/Engine/Extras/RoboMerge/v3/src/robo/edgebot.ts
@@ -186,6 +186,7 @@ class EdgeBotImpl extends PerforceStatefulBot {
 
 	protected _forceSetLastCl_NoReset(value: number) {
 		this.gate.numChangesRemaining = 0
+		this.gate.setEdgeLastCl(value)
 		return super._forceSetLastCl_NoReset(value)
 	}
 
diff --git a/Engine/Extras/RoboMerge/v3/src/robo/gate.ts b/Engine/Extras/RoboMerge/v3/src/robo/gate.ts
--- a/Engine/Extras/RoboMerge/v3/src/robo/gate.ts
+++ b/Engine/Extras/RoboMerge/v3/src/robo/gate.ts
@@ -644,7 +644,6 @@ export class Gate {
 		if (this.queuedGates.length > 0) {
 			data.queued = this.queuedGates
 		}
-		data.lastCl = this.lastCl
 		this.persistence.set(GATE_INFO_KEY, data)
 		this.persistence.set(BYPASS_KEY, this.bypass)
 	}
@@ -659,14 +658,11 @@ export class Gate {
 			this.context.logger.info(`Restoring saved gate info: current ${saved.current && saved.current.cl}`)
 
 			const version = saved.version || 0
-			if (saved.lastCl) {
-				this.setLastCl(saved.lastCl)
-			}
 
 			if (saved.current) {
 				this.currentGateInfo = saved.current
 
-				if (saved.lastCl && saved.lastCl > saved.current.cl) {
+				if (this.lastCl < saved.current.cl) {
 					this.reportCatchingUp()
 				}
 			}
@@ -695,6 +691,10 @@ export class Gate {
 		return this.eventContext.edgeLastCl
 	}
 
+	setEdgeLastCl(value: number) {
+		this.eventContext.edgeLastCl = value
+	}
+
 	private setLastCl(cl: number) {
 		this.eventContext.edgeLastCl = cl
 	}
