Fix erroneous error message reported to user when non-blocking approval flow is followed
Enable approvals test and add a non blocking version

diff --git a/Engine/Extras/RoboMerge/v3/functional_tests/framework/src/tests.ts b/Engine/Extras/RoboMerge/v3/functional_tests/framework/src/tests.ts
--- a/Engine/Extras/RoboMerge/v3/functional_tests/framework/src/tests.ts
+++ b/Engine/Extras/RoboMerge/v3/functional_tests/framework/src/tests.ts
@@ -4,6 +4,7 @@ import { DEFAULT_BOT_SETTINGS, EdgeProperties, FunctionalTest, getRootDataClient
 import * as bent from 'bent'
 import { Perforce } from './test-perforce'
 import { GenericTargetTest } from './GenericTargetTest'
+import { Approval, ApprovalNonBlocking } from './tests/approval'
 import { BlockAssets } from './tests/block-assets'
 import { BlockIgnore } from './tests/block-ignore'
 import { BranchspecTest } from './tests/branchspec'
@@ -230,7 +231,10 @@ async function go() {
 
 		new ExcludeDescriptionsPerEdge(p4),
 		new SkipIfUnchanged(p4),
-		new TestEdgeBlockAssetsReconsider(p4)
+		new TestEdgeBlockAssetsReconsider(p4),
+		new Approval(p4),
+		new ApprovalNonBlocking(p4)  //55
+
 	]
 
 	const TARGET_TEST_DEFS: [string[], string | null][] = [
diff --git a/Engine/Extras/RoboMerge/v3/functional_tests/framework/src/tests/approval.ts b/Engine/Extras/RoboMerge/v3/functional_tests/framework/src/tests/approval.ts
--- a/Engine/Extras/RoboMerge/v3/functional_tests/framework/src/tests/approval.ts
+++ b/Engine/Extras/RoboMerge/v3/functional_tests/framework/src/tests/approval.ts
@@ -37,4 +37,46 @@ export class Approval extends SimpleMainAndReleaseTestBase {
 		  }
 		]
 	}
+}
+
+export class ApprovalNonBlocking extends SimpleMainAndReleaseTestBase {
+	cl = -1
+
+	async setup() {
+		await this.p4.depot('stream', this.depotSpec())
+		await this.createStreamsAndWorkspaces(streams)
+		await P4Util.addFileAndSubmit(this.getClient('Main'), 'test.txt', 'Initial content')
+		await this.initialPopulate()
+	}
+
+	async run() {
+		const releaseClient = this.getClient('Release')
+		await releaseClient.sync()
+		this.cl = await P4Util.editFileAndSubmit(releaseClient, 'test.txt', 'Initial content\n\nmergeable')
+	}
+
+	async verify() {
+		const approvalMessageSent = await this.wasMessagePostedToSlack('abc123', this.cl, 'Main')
+		if (!approvalMessageSent) {
+			throw new Error(`Approval message not sent to channel 'abc123' for CL ${this.cl}`)
+		}
+		return Promise.all([
+			this.ensureNotBlocked('Main'),
+			this.ensureNotBlocked('Release', 'Main'),
+			// CL should be shelved for approval, not committed to Main
+			this.checkHeadRevision('Main', 'test.txt', 1)
+		])
+	}
+
+	getEdges() : EdgeProperties[] {
+		return [
+		  { from: this.fullBranchName('Release'), to: this.fullBranchName('Main')
+		  , approval: {
+		      description: 'hey there!',
+		      channelId: 'abc123',
+		      block: false
+		    }
+		  }
+		]
+	}
 }
\ No newline at end of file
diff --git a/Engine/Extras/RoboMerge/v3/src/robo/edgebot.ts b/Engine/Extras/RoboMerge/v3/src/robo/edgebot.ts
--- a/Engine/Extras/RoboMerge/v3/src/robo/edgebot.ts
+++ b/Engine/Extras/RoboMerge/v3/src/robo/edgebot.ts
@@ -899,6 +899,7 @@ class EdgeBotImpl extends PerforceStatefulBot {
 							shelfCl: pending.newCl
 						});
 				}
+				return new EdgeIntegrationDetails('shelved', pending.newCl)
 			}
 			else if (pending.action.flags.has('#filter_blocked#')) {
 
diff --git a/Engine/Extras/RoboMerge/v3/src/robo/nodebot.ts b/Engine/Extras/RoboMerge/v3/src/robo/nodebot.ts
--- a/Engine/Extras/RoboMerge/v3/src/robo/nodebot.ts
+++ b/Engine/Extras/RoboMerge/v3/src/robo/nodebot.ts
@@ -2388,13 +2388,15 @@ export class NodeBot extends PerforceStatefulBot implements NodeBotInterface {
 						  `in ${pending.action.branch.name} needs to be approved.\n\n` +
 						  approval.description
 
-			const opts: SlackMessage = { 
-				title:'', 
-				text: message, 
-				style: SlackMessageStyles.DANGER, 
+			const opts: SlackMessage = {
+				title:'',
+				text: message,
+				style: SlackMessageStyles.DANGER,
 				fields,
 				mrkdwn: true,
-				channel: approval.channelId
+				channel: approval.channelId,
+				cl: pending.change.cl,
+				target: pending.action.branch.upperName
 			}
 
 			this.slackMessages.postNonConflictMessage(opts)
