Ensure the working directory always exists before running commands that rely on it

diff --git a/Engine/Extras/RoboMerge/v3/src/common/perforce.ts b/Engine/Extras/RoboMerge/v3/src/common/perforce.ts
--- a/Engine/Extras/RoboMerge/v3/src/common/perforce.ts
+++ b/Engine/Extras/RoboMerge/v3/src/common/perforce.ts
@@ -1,5 +1,6 @@
 // Copyright Epic Games, Inc. All Rights Reserved.
 import { execFile, ExecFileOptions } from 'child_process';
+import * as fs from 'fs';
 import * as path from 'path';
 // should really separate out RM-specific operations, but no real gain at the moment
 import { isUserAKnownBot, isUserSlackGroup, postToRobomergeAlerts } from '../robo/notifications';
@@ -2144,6 +2145,9 @@ export class PerforceContext {
 		// we need to run within the workspace directory so p4 selects the correct AltRoot
 		let options: ExecFileOptions = { maxBuffer: 500 * 1024 * 1024 };
 		if (workspace && workspace.directory && !opts.noCwd) {
+			if (!fs.existsSync(workspace.directory)) {
+				fs.mkdirSync(workspace.directory, {recursive: true});
+			}
 			options.cwd = workspace.directory;
 		}
 
diff --git a/Engine/Extras/RoboMerge/v3/src/robo/edgebot.ts b/Engine/Extras/RoboMerge/v3/src/robo/edgebot.ts
--- a/Engine/Extras/RoboMerge/v3/src/robo/edgebot.ts
+++ b/Engine/Extras/RoboMerge/v3/src/robo/edgebot.ts
@@ -1,6 +1,5 @@
 // Copyright Epic Games, Inc. All Rights Reserved.
 
-import * as fs from 'fs';
 import * as p4util from '../common/p4util';
 
 import { ContextualLogger } from "../common/logger";
@@ -517,12 +516,7 @@ class EdgeBotImpl extends PerforceStatefulBot {
 			workspaceName += '_' + edgeServer.id.toUpperCase()
 		}
 
-		// ensure root directory exists (we set the root diretory to be the cwd)
 		const path = getRootDirectoryForBranch(workspaceName);
-		if (!fs.existsSync(path)) {
-			this.logger.info(`Making directory ${path}`);
-			fs.mkdirSync(path);
-		}
 
 		// do we already have a workspace?
 		const existingWorkspaceInfo = await p4.find_workspace_by_name(workspaceName, {edgeServerAddress, includeUnloaded: true})
