Be explicit about sending a skip message from the user so it doesn't accidentally get conflated with changing the last CL in rare cases where the conflict is cleared just as a skip message is sent through.

diff --git a/Engine/Extras/RoboMerge/v3/public/js/skip.js b/Engine/Extras/RoboMerge/v3/public/js/skip.js
--- a/Engine/Extras/RoboMerge/v3/public/js/skip.js
+++ b/Engine/Extras/RoboMerge/v3/public/js/skip.js
@@ -180,7 +180,8 @@ function doSkip(fullData, edgeName, branchCl, skipText) {
         // If the forced CL is the same as the conflict CL, the conflict will be skipped
         cl: branchCl,
         reason: skipText,
-        unblock: 'true'
+        unblock: 'true',
+        isSkip: 'true'
     }
 
     let branchData = fullData.branch
diff --git a/Engine/Extras/RoboMerge/v3/src/robo/bot-interfaces.ts b/Engine/Extras/RoboMerge/v3/src/robo/bot-interfaces.ts
--- a/Engine/Extras/RoboMerge/v3/src/robo/bot-interfaces.ts
+++ b/Engine/Extras/RoboMerge/v3/src/robo/bot-interfaces.ts
@@ -51,7 +51,7 @@ export interface IPCControls {
 	acknowledge(acknowledger: string, changeCl: number) : OperationResult
 	unacknowledge(changeCl: number) : OperationResult
 	
-	forceSetLastClWithContext(value: number, culprit: string, reason: string, unblock?: boolean): number
+	forceSetLastClWithContext(value: number, culprit: string, reason: string, unblock?: boolean, isSkip?: boolean): number
 	setGateCl(value: number, culprit: string, reason: string): Promise<number | null>
 }
 
diff --git a/Engine/Extras/RoboMerge/v3/src/robo/branch-interfaces.ts b/Engine/Extras/RoboMerge/v3/src/robo/branch-interfaces.ts
--- a/Engine/Extras/RoboMerge/v3/src/robo/branch-interfaces.ts
+++ b/Engine/Extras/RoboMerge/v3/src/robo/branch-interfaces.ts
@@ -289,6 +289,7 @@ export interface ForcedCl {
 	previousCl: number
 	culprit: string
 	reason: string
+	isSkip?: boolean
 }
 
 export type GateInfo = {
diff --git a/Engine/Extras/RoboMerge/v3/src/robo/conflicts.ts b/Engine/Extras/RoboMerge/v3/src/robo/conflicts.ts
--- a/Engine/Extras/RoboMerge/v3/src/robo/conflicts.ts
+++ b/Engine/Extras/RoboMerge/v3/src/robo/conflicts.ts
@@ -386,11 +386,12 @@ export class Conflicts {
 		}
 	}
 
-	// If a forced CL is the same as a conflict CL, the conflict as skipped.
-	// Otherwise we'll notify that the branch has been forced to a CL
+	// If a forced CL is the same as a conflict CL, the conflict is skipped.
+	// Otherwise we'll notify that the branch has been forced to a CL, unless the
+	// caller flagged this as a skip (in which case the conflict was already gone).
 	onForcedLastCl(forced: ForcedCl) {
 		this.eventTriggers.reportForcedLastCl(forced)
-		if (!this.notifyConflictSkipped(forced)) {
+		if (!this.notifyConflictSkipped(forced) && !forced.isSkip) {
 			this.eventTriggers.reportNonSkipLastClChange(forced)
 		}
 	}
diff --git a/Engine/Extras/RoboMerge/v3/src/robo/edgebot.ts b/Engine/Extras/RoboMerge/v3/src/robo/edgebot.ts
--- a/Engine/Extras/RoboMerge/v3/src/robo/edgebot.ts
+++ b/Engine/Extras/RoboMerge/v3/src/robo/edgebot.ts
@@ -1206,7 +1206,7 @@ class EdgeBotImpl extends PerforceStatefulBot {
 		return status as EdgeStatusFields
 	}
 
-	public forceSetLastClWithContext(value: number, culprit: string, reason: string, unblock?: boolean) {
+	public forceSetLastClWithContext(value: number, culprit: string, reason: string, unblock?: boolean, isSkip?: boolean) {
 		const prevValue = this._forceSetLastCl_NoReset(value)
 
 		if (unblock) {
@@ -1214,8 +1214,8 @@ class EdgeBotImpl extends PerforceStatefulBot {
 		}
 
 		// trigger events
-		this.sourceNode.onForcedLastCl(this.displayName, this.targetBranch.upperName, value, prevValue, culprit, reason)
-		
+		this.sourceNode.onForcedLastCl(this.displayName, this.targetBranch.upperName, value, prevValue, culprit, reason, isSkip)
+
 		return prevValue
 	}
 
diff --git a/Engine/Extras/RoboMerge/v3/src/robo/ipc.ts b/Engine/Extras/RoboMerge/v3/src/robo/ipc.ts
--- a/Engine/Extras/RoboMerge/v3/src/robo/ipc.ts
+++ b/Engine/Extras/RoboMerge/v3/src/robo/ipc.ts
@@ -444,7 +444,7 @@ export class IPC {
 				return {statusCode: 400, message: 'Invalid CL parameter: ' + cl}
 			}
 
-			let prevCl = generalOpTarget.forceSetLastClWithContext(cl, query.who, query.reason, true)
+			let prevCl = generalOpTarget.forceSetLastClWithContext(cl, query.who, query.reason, true, query.isSkip === 'true')
 
 			this.ipcLogger.info(`Forcing last CL=${cl} on ${botname} : ${branch.name} (was CL ${prevCl}), ` +
 														`requested by ${query.who} (Reason: ${query.reason})`)
diff --git a/Engine/Extras/RoboMerge/v3/src/robo/nodebot.ts b/Engine/Extras/RoboMerge/v3/src/robo/nodebot.ts
--- a/Engine/Extras/RoboMerge/v3/src/robo/nodebot.ts
+++ b/Engine/Extras/RoboMerge/v3/src/robo/nodebot.ts
@@ -600,8 +600,8 @@ export class NodeBot extends PerforceStatefulBot implements NodeBotInterface {
 		throw err
 	}
 
-	onForcedLastCl(nodeOrEdgeName: string, targetBranchUpperName: string, forcedCl: number, previousCl: number, culprit: string, reason: string) {
-		this.conflicts.onForcedLastCl({nodeOrEdgeName, sourceBranchUpperName: this.branch.upperName, targetBranchUpperName, forcedCl, previousCl, culprit, reason})
+	onForcedLastCl(nodeOrEdgeName: string, targetBranchUpperName: string, forcedCl: number, previousCl: number, culprit: string, reason: string, isSkip?: boolean) {
+		this.conflicts.onForcedLastCl({nodeOrEdgeName, sourceBranchUpperName: this.branch.upperName, targetBranchUpperName, forcedCl, previousCl, culprit, reason, isSkip})
 	}
 
 	persistQueuedChanges() {
diff --git a/Engine/Extras/RoboMerge/v3/src/robo/perforce-stateful-bot.ts b/Engine/Extras/RoboMerge/v3/src/robo/perforce-stateful-bot.ts
--- a/Engine/Extras/RoboMerge/v3/src/robo/perforce-stateful-bot.ts
+++ b/Engine/Extras/RoboMerge/v3/src/robo/perforce-stateful-bot.ts
@@ -151,7 +151,7 @@ export abstract class PerforceStatefulBot implements BotIPC {
 		bot.lastCl = change.change;
 	}
 
-	abstract forceSetLastClWithContext(value: number, culprit: string, reason: string, unblock: boolean): number;
+	abstract forceSetLastClWithContext(value: number, culprit: string, reason: string, unblock: boolean, isSkip?: boolean): number;
 	abstract setGateCl(value: number, culprit: string, reason: string): Promise<number | null>;
 
 	async _getChange(changeCl: number) : Promise<Change | string> {
