Allow the specification of a display name for a node

diff --git a/Engine/Extras/RoboMerge/v3/public/js/boilerplate.js b/Engine/Extras/RoboMerge/v3/public/js/boilerplate.js
--- a/Engine/Extras/RoboMerge/v3/public/js/boilerplate.js
+++ b/Engine/Extras/RoboMerge/v3/public/js/boilerplate.js
@@ -1213,7 +1213,7 @@ function postRenderNameCell_Node(nameCell, nodeData, includeCollapseControls, co
 		
 		const name = getDisplayName(nodeData)
 		collapseLink = $('<a role="button" data-toggle="collapse" title="Show/Hide Edges">').appendTo(nameHeader)
-		collapseLink.attr('data-target', '#' + getNodeEdgesRowId(nodeData.bot, name))
+		collapseLink.attr('data-target', '#' + getNodeEdgesRowId(nodeData.bot, nodeData.def.name))
 
 		let icon = $('<i class="fas fa-lg">').appendTo(collapseLink)
 		
diff --git a/Engine/Extras/RoboMerge/v3/public/js/flow.graph.js b/Engine/Extras/RoboMerge/v3/public/js/flow.graph.js
--- a/Engine/Extras/RoboMerge/v3/public/js/flow.graph.js
+++ b/Engine/Extras/RoboMerge/v3/public/js/flow.graph.js
@@ -171,9 +171,10 @@ class Graph {
             return;
         let tooltip = `${branchStatus.bot} - ${branch.rootPath}`;
         if (branch.aliases && branch.aliases.length !== 0) {
-            tooltip += ` (${branch.aliases.join(', ')})`;
+            const displayAliases = branch.aliases.map(a => a === branch.upperName ? branchStatus.display_name : a);
+            tooltip += ` (${displayAliases.join(', ')})`;
         }
-        const info = new Info(`_${branchStatus.bot}_${branch.upperName.replace(/[^\w]+/g, '_')}`, tooltip, branch.name, [branchStatus]);
+        const info = new Info(`_${branchStatus.bot}_${branch.upperName.replace(/[^\w]+/g, '_')}`, tooltip, branchStatus.display_name, [branchStatus]);
         info.title = title
         this.allInfos.push(info);
         if (branch.config.graphNodeColor) {
@@ -263,7 +264,8 @@ class Graph {
                 for (const branchStatus of streams) {
                     let tooltip = `${branchStatus.bot} - ${branchStatus.def.rootPath}`;
                     if (branchStatus.def.aliases && branchStatus.def.aliases.length !== 0) {
-                        tooltip += ` (${branchStatus.def.aliases.join(', ')})`;
+                        const displayAliases = branchStatus.def.aliases.map(a => a === branchStatus.def.upperName ? branchStatus.display_name : a);
+                        tooltip += ` (${displayAliases.join(', ')})`;
                     }
                     globalToolTip += tooltip + "\n";
                 }
diff --git a/Engine/Extras/RoboMerge/v3/src/robo/branchdefs.ts b/Engine/Extras/RoboMerge/v3/src/robo/branchdefs.ts
--- a/Engine/Extras/RoboMerge/v3/src/robo/branchdefs.ts
+++ b/Engine/Extras/RoboMerge/v3/src/robo/branchdefs.ts
@@ -128,6 +128,7 @@ const nodeOptionFieldsPrototype = {
 	uniqueBranch: false,
 
 	graphNodeColor: '',
+	displayName: '',
 
 	additionalSlackChannelForBlockages: '',
 	postMessagesToAdditionalChannelOnly: false,
diff --git a/Engine/Extras/RoboMerge/v3/src/robo/edgebot.ts b/Engine/Extras/RoboMerge/v3/src/robo/edgebot.ts
--- a/Engine/Extras/RoboMerge/v3/src/robo/edgebot.ts
+++ b/Engine/Extras/RoboMerge/v3/src/robo/edgebot.ts
@@ -149,7 +149,7 @@ class EdgeBotImpl extends PerforceStatefulBot {
 		return `${this.sourceNode.fullNameForLogging}=>${this.targetBranch.name}`
 	}
 	get displayName() {
-		return `${this.sourceNode.displayName} -> ${this.targetBranch.name}`
+		return `${this.sourceNode.displayName} -> ${this.targetBranch.config.displayName || this.targetBranch.name}`
 	}
 	get incognitoMode() {
 		return !!this.options.incognitoMode
diff --git a/Engine/Extras/RoboMerge/v3/src/robo/nodebot.ts b/Engine/Extras/RoboMerge/v3/src/robo/nodebot.ts
--- a/Engine/Extras/RoboMerge/v3/src/robo/nodebot.ts
+++ b/Engine/Extras/RoboMerge/v3/src/robo/nodebot.ts
@@ -544,7 +544,7 @@ export class NodeBot extends PerforceStatefulBot implements NodeBotInterface {
 	}
 
 	get displayName() {
-		return this.branch.name
+		return this.branch.config.displayName || this.branch.name
 	}
 	get fullName() {
 		return getNodeBotFullName(this.graphBotName, this.branch.name)
diff --git a/Engine/Extras/RoboMerge/v3/src/robo/status.ts b/Engine/Extras/RoboMerge/v3/src/robo/status.ts
--- a/Engine/Extras/RoboMerge/v3/src/robo/status.ts
+++ b/Engine/Extras/RoboMerge/v3/src/robo/status.ts
@@ -19,6 +19,10 @@ export class Status {
 		if (branch.bot) {
 			await branch.bot.applyStatus(status)
 		}
+		else {
+			// No running bot (e.g. preview) — set display_name from config or fall back to branch name
+			status.display_name = (branch.config.displayName as string | undefined) || branch.name
+		}
 	}
 
 	addWarning(warning: string) {
