Remove internal only github/sync calls from clientside trackchanges js
#rnx

diff --git a/Engine/Extras/RoboMerge/v3/public/js/trackchange.js b/Engine/Extras/RoboMerge/v3/public/js/trackchange.js
--- a/Engine/Extras/RoboMerge/v3/public/js/trackchange.js
+++ b/Engine/Extras/RoboMerge/v3/public/js/trackchange.js
@@ -72,38 +72,8 @@ function getStreamGraphName(streamDisplayName) {
 	return match ? match[1] : streamDisplayName
 }
 
-async function getGitHubCommit(cl) {
-	const timeout = new Promise((resolve, reject) => {
-    	setTimeout(() => reject(new Error('GitSync request took longer than 3 seconds')), 3000)
-    });
-     const response = await Promise.race([
-		$.get(`https://gitsync.devtools.epicgames.com/api/v1/changes/${cl}`),
-        timeout
-    ]);
-    return {"commit":response.commit, "commitURL":response.commitURL};
-}
-
-async function getGitHubCommits(changes) {
-	return new Map(await Promise.all(Object.keys(changes).map(async cl => {
-		try {
-            const ghCommit = await getGitHubCommit(cl);
-            return [cl, ghCommit];
-        } catch(error) {
-            return [cl, null];
-        }
-	})));
-}
-
-function formatGitHubCommit(change) {
-	if (change) {
-		return `<td><a href="${change.commitURL}" target="_blank"><button class="btn btn-sm">${change.commit.slice(0,6)}</button></a></td>`
-	}
-	return '<td></td>'
-}
-
-async function generateChangeList(dataObj) {
+function generateChangeList(dataObj) {
 
-	const gitHubCommits = await getGitHubCommits(dataObj.changes)
 	let html = '<div style="margin: auto; width: 80%;"><table class="table"><tbody>'
 	// By default the keys are going to be numerically ordered, but we want changes to appear after their parent
 	// So go through and build them in a list order we prefer
@@ -145,7 +115,6 @@ async function generateChangeList(dataObj) {
 		} else {
 			html += `<td>CL#${cl}</td>`
 		}
-		html += `${formatGitHubCommit(gitHubCommits.get(clKey))}`
 		html += formatReconsiderButton(cl, changeDetails.reconsiderInfo)
 		html += '</tr>'
 	}
@@ -210,10 +179,10 @@ function createGraph(changes) {
     return graphContainer;	
 }
 
-async function buildResults(data) {
+function buildResults(data) {
 	if (Object.keys(data.data.changes).length > 0) {
 		const $successPanel = $('#success-panel');
-		$('#changes', $successPanel).html(await generateChangeList(data.data));
+		$('#changes', $successPanel).html(generateChangeList(data.data));
 		$successPanel.show();
 		$('#graph').append(createGraph(data.data.changes))
 	} else {
