Add EditableBranch displayName accessor
Clean up displayName accesses
Fix up notifications using name instead of displayName

diff --git a/Engine/Extras/RoboMerge/v3/src/robo/branch-interfaces.ts b/Engine/Extras/RoboMerge/v3/src/robo/branch-interfaces.ts
--- a/Engine/Extras/RoboMerge/v3/src/robo/branch-interfaces.ts
+++ b/Engine/Extras/RoboMerge/v3/src/robo/branch-interfaces.ts
@@ -14,6 +14,9 @@ export function resolveBranchArg(branchArg: BranchArg, toUpperCase?: boolean): s
 		(branchArg as Branch).upperName || (branchArg as string).toUpperCase() :
 		(branchArg as Branch).name || branchArg as string
 }
+export function getDisplayNameFromBranchArg(branchArg: BranchArg): string {
+	return (branchArg as Branch).displayName || branchArg as string
+}
 export function branchesMatch(branchArg1: BranchArg | null | undefined, branchArg2: BranchArg | null | undefined) {
 	return !branchArg1 || !branchArg2 || (typeof(branchArg1) === 'object' && typeof(branchArg2) === 'object') ?
 		branchArg1 === branchArg2 : resolveBranchArg(branchArg1, true) === resolveBranchArg(branchArg2, true)
@@ -133,6 +136,7 @@ export type EditableBranch = BranchBase & {
 
 	edgeProperties: Map<string, EdgeOptions>
 
+	displayName: string // computed from config.displayName or name
 	isMonitored: boolean // property
 }
 
diff --git a/Engine/Extras/RoboMerge/v3/src/robo/branchgraph.ts b/Engine/Extras/RoboMerge/v3/src/robo/branchgraph.ts
--- a/Engine/Extras/RoboMerge/v3/src/robo/branchgraph.ts
+++ b/Engine/Extras/RoboMerge/v3/src/robo/branchgraph.ts
@@ -250,6 +250,7 @@ if (botname === '__TEST__') {
 			nagAcknowledgedLeeway: options.nagAcknowledgedLeeway || null,
 			nagWhenBlocked: options.nagWhenBlocked || null,
 			visibility: options.visibility || this.config.visibility,
+			get displayName() { return this.config.displayName || this.name },
 			get isMonitored() { return !!(this.bot && this.bot.isRunning) }
 		}
 		if (branch.changelistParser) {
diff --git a/Engine/Extras/RoboMerge/v3/src/robo/edgebot.ts b/Engine/Extras/RoboMerge/v3/src/robo/edgebot.ts
--- a/Engine/Extras/RoboMerge/v3/src/robo/edgebot.ts
+++ b/Engine/Extras/RoboMerge/v3/src/robo/edgebot.ts
@@ -149,7 +149,7 @@ class EdgeBotImpl extends PerforceStatefulBot {
 		return `${this.sourceNode.fullNameForLogging}=>${this.targetBranch.name}`
 	}
 	get displayName() {
-		return `${this.sourceNode.displayName} -> ${this.targetBranch.config.displayName || this.targetBranch.name}`
+		return `${this.sourceNode.displayName} -> ${this.targetBranch.displayName}`
 	}
 	get incognitoMode() {
 		return !!this.options.incognitoMode
diff --git a/Engine/Extras/RoboMerge/v3/src/robo/nodebot.ts b/Engine/Extras/RoboMerge/v3/src/robo/nodebot.ts
--- a/Engine/Extras/RoboMerge/v3/src/robo/nodebot.ts
+++ b/Engine/Extras/RoboMerge/v3/src/robo/nodebot.ts
@@ -544,7 +544,7 @@ export class NodeBot extends PerforceStatefulBot implements NodeBotInterface {
 	}
 
 	get displayName() {
-		return this.branch.config.displayName || this.branch.name
+		return this.branch.displayName
 	}
 	get fullName() {
 		return getNodeBotFullName(this.graphBotName, this.branch.name)
diff --git a/Engine/Extras/RoboMerge/v3/src/robo/notifications.ts b/Engine/Extras/RoboMerge/v3/src/robo/notifications.ts
--- a/Engine/Extras/RoboMerge/v3/src/robo/notifications.ts
+++ b/Engine/Extras/RoboMerge/v3/src/robo/notifications.ts
@@ -6,7 +6,7 @@ import { Badge } from '../common/badge';
 import { Random, setDefault } from '../common/helper';
 import { ContextualLogger } from '../common/logger';
 import { ExclusiveFileDetails } from '../common/perforce';
-import { Blockage, Branch, BranchArg, EditableBranch, ExclusiveLockInfo, ForcedCl, MergeAction, NagBlockageEvent, NodeOpUrlGenerator, PauseEvent, resolveBranchArg } from './branch-interfaces';
+import { Blockage, Branch, BranchArg, EditableBranch, ExclusiveLockInfo, ForcedCl, MergeAction, NagBlockageEvent, NodeOpUrlGenerator, PauseEvent, resolveBranchArg, getDisplayNameFromBranchArg } from './branch-interfaces';
 import { PersistentConflict, Resolution } from './conflict-interfaces';
 import { BlockageStatusChangedEvent, BotEventHandler, BotEvents } from './events';
 import { NodeBot } from './nodebot';
@@ -558,7 +558,7 @@ export class BotNotifications implements BotEventHandler {
 		}														
 
 		let message = this.makeSlackChannelMessage(
-			`${sourceBranch.name} blocked! (${issue})`,
+			`${sourceBranch.displayName} blocked! (${issue})`,
 			text,
 			SlackMessageStyles.DANGER, 
 			makeClLink(changeInfo.cl), 
@@ -622,7 +622,7 @@ export class BotNotifications implements BotEventHandler {
 			let dm: SlackMessage
 			if (blockage.approval) {
 				dm = {
-					title: 'Approval needed to commit to ' + targetBranch.name,
+					title: 'Approval needed to commit to ' + targetBranch.displayName,
 					text: `Your change has been shelved in ${makeClLink(blockage.approval.shelfCl)} and sent to <#${blockage.approval.settings.channelId}> for approval\n\n` +
 							blockage.approval.settings.description,
 					channel: "",
@@ -631,7 +631,7 @@ export class BotNotifications implements BotEventHandler {
 			}
 			else {
 				const dmText = `Your change (${makeClLink(changeInfo.source_cl)}) ` +
-					`hit '${issue}' while merging from *${sourceBranch.name}* to *${targetBranch.name}*.\n\n` +
+					`hit '${issue}' while merging from *${sourceBranch.displayName}* to *${targetBranch.displayName}*.\n\n` +
 					'`' + blockage.change.description.substring(0, 80) + '`\n\n' +
 					"*_Resolving this blockage is time sensitive._*"
 				
@@ -642,7 +642,7 @@ export class BotNotifications implements BotEventHandler {
 					throw error
 				}
 
-				dm = this.makeSlackDirectMessage(dmText, changeInfo.cl, changeInfo.cl, targetBranch.name, urls, threadLinks)
+				dm = this.makeSlackDirectMessage(dmText, changeInfo.cl, changeInfo.cl, targetBranch.displayName, urls, threadLinks)
 			}
 
 			this.slackMessages.postDM(userEmail, changeInfo.cl, targetBranch, dm)
@@ -967,9 +967,9 @@ export class BotNotifications implements BotEventHandler {
 
 	private makeSlackChannelMessage(title: string, text: string, style: SlackMessageStyles, clDesc: string, sourceBranch: BranchArg,
 											targetBranch?: BranchArg, author?: string, buttons?: SlackLinkButtonsAttachment[]) {
-		const integrationText = [resolveBranchArg(sourceBranch)]
+		const integrationText = [getDisplayNameFromBranchArg(sourceBranch)]
 		if (targetBranch) {
-			integrationText.push(resolveBranchArg(targetBranch))
+			integrationText.push(getDisplayNameFromBranchArg(targetBranch))
 		}
 		const fields: SlackMessageField[] = [
 			{title: 'Integration', short: true, value: integrationText.join(' -> ')},
diff --git a/Engine/Extras/RoboMerge/v3/src/robo/status.ts b/Engine/Extras/RoboMerge/v3/src/robo/status.ts
--- a/Engine/Extras/RoboMerge/v3/src/robo/status.ts
+++ b/Engine/Extras/RoboMerge/v3/src/robo/status.ts
@@ -21,7 +21,7 @@ export class Status {
 		}
 		else {
 			// No running bot (e.g. preview) — set display_name from config or fall back to branch name
-			status.display_name = (branch.config.displayName as string | undefined) || branch.name
+			status.display_name = branch.displayName
 		}
 	}
 
