Fix conflicts not being canclled when setting last CL

diff --git a/Engine/Extras/RoboMerge/v3/src/robo/conflicts.ts b/Engine/Extras/RoboMerge/v3/src/robo/conflicts.ts
--- a/Engine/Extras/RoboMerge/v3/src/robo/conflicts.ts
+++ b/Engine/Extras/RoboMerge/v3/src/robo/conflicts.ts
@@ -407,39 +407,63 @@ export class Conflicts {
 	}
 
 	private notifyConflictSkipped(forced: ForcedCl) {
-		if (forced.forcedCl <= forced.previousCl) {
-			return false
-		}
+		const isForward = forced.forcedCl > forced.previousCl
+		const isBackward = forced.forcedCl < forced.previousCl
 
 		const remainingConflicts: PersistentConflict[] = []
-		const unblocksToReport: PersistentConflict[] = []
+		const skippedToReport: PersistentConflict[] = []
+		const cancelledToReport: PersistentConflict[] = []
 
 		for (const conflict of this.conflicts) {
-			// If the forced CL is the same as our conflict CL, count this conflict as skipped
+			const onSameEdge = conflict.blockedBranchName === forced.sourceBranchUpperName &&
+				conflict.targetBranchName === forced.targetBranchUpperName
+			let resolution: Resolution = Resolution.DUNNO
+
 			if (conflict.cl === forced.forcedCl) {
-				conflict.resolvingAuthor = forced.culprit
-				conflict.resolution = Resolution.SKIPPED
+				resolution = Resolution.SKIPPED
 				conflict.resolvingReason = forced.reason
+				skippedToReport.push(conflict)
+			}
+			else if (onSameEdge) {
+				if (isForward && conflict.cl < forced.forcedCl) {
+					resolution = Resolution.CANCELLED
+					conflict.resolvingReason = 'Bot position was advanced past this CL'
+					cancelledToReport.push(conflict)
+				}
+				else if (isBackward && conflict.cl > forced.forcedCl) {
+					resolution = Resolution.CANCELLED
+					conflict.resolvingReason = 'Bot position was reset to an earlier CL'
+					cancelledToReport.push(conflict)
+				}
+			}
+			if (resolution != Resolution.DUNNO) {
+				conflict.resolvingAuthor = forced.culprit
+				conflict.resolution = resolution
 				setResolvedTimeAsNow(conflict)
-				unblocksToReport.push(conflict)
 			}
 			else {
 				remainingConflicts.push(conflict)
 			}
 		}
 
-		if (unblocksToReport.length === 0) {
+		if (skippedToReport.length === 0 && cancelledToReport.length === 0) {
 			return false
 		}
 
 		this.conflictLogger.debug(`notifyConflictSkipped -- ${forced.nodeOrEdgeName} @ CL# ${forced.forcedCl}`)
 		this.setConflicts(remainingConflicts)
 
-		for (const unblock of unblocksToReport) {
+		for (const unblock of skippedToReport) {
 			this.conflictLogger.info('Reporting unblocked (skipped) ' + unblock.cl)
 			this.reportUnblocked(unblock)
 		}
-		return true
+		for (const unblock of cancelledToReport) {
+			this.conflictLogger.info('Reporting unblocked (cancelled) ' + unblock.cl)
+			this.reportUnblocked(unblock)
+		}
+
+		// Only suppress reportNonSkipLastClChange when a conflict was explicitly skipped
+		return skippedToReport.length > 0
 	}
 
 	persist() {
