Remove dead code path

diff --git a/Engine/Extras/RoboMerge/v3/src/robo/nodebot.ts b/Engine/Extras/RoboMerge/v3/src/robo/nodebot.ts
--- a/Engine/Extras/RoboMerge/v3/src/robo/nodebot.ts
+++ b/Engine/Extras/RoboMerge/v3/src/robo/nodebot.ts
@@ -345,23 +345,8 @@ export class NodeBot extends PerforceStatefulBot implements NodeBotInterface {
 		}
 		else {
 			this.ticksSinceLastNewP4Commit = 0
-
 			this.headCL = changes[0].change
 
-			// list of changes except reconsiders
-			const changesToProcess: Change[] = []
-
-			// process manual changes even if paused or above the lastGoodCL
-			for (const change of changes) {
-				if (change.isUserRequest) {
-					this.nodeBotLogger.debug(`Processing manual change CL ${change.change}`)
-					await this._processAndMergeCl(availableEdges, change, true)
-				}
-				else {
-					changesToProcess.push(change)
-				}
-			}
-
 			if (!this.isAvailable) {
 				this.nodeBotLogger.verbose(`Not processing changes due to unavailability.`)
 				this.nodeBotLogger.debug(`Blocked? ${this.pauseState.isBlocked()}: ${this.pauseState.blockagePauseInfo}`)
@@ -370,8 +355,8 @@ export class NodeBot extends PerforceStatefulBot implements NodeBotInterface {
 			}
 			else {
 				// make sure the list is sorted in ascending order
-				changesToProcess.sort((a, b) => a.change - b.change)
-				await Promise.all(Array.from(availableEdgesAtCL.entries()).map(async ([minCL,edgesAtCL]) => this._processListOfChanges(edgesAtCL, minCL, changesToProcess)));
+				changes.sort((a, b) => a.change - b.change)
+				await Promise.all(Array.from(availableEdgesAtCL.entries()).map(async ([minCL,edgesAtCL]) => this._processListOfChanges(edgesAtCL, minCL, changes)));
 			}
 		}
 		this.sendNagNotifications()
