Clean up and encapsulate vault and args a little more in to RoboArgs

diff --git a/Engine/Extras/RoboMerge/v3/src/robo/notifications.ts b/Engine/Extras/RoboMerge/v3/src/robo/notifications.ts
--- a/Engine/Extras/RoboMerge/v3/src/robo/notifications.ts
+++ b/Engine/Extras/RoboMerge/v3/src/robo/notifications.ts
@@ -1,7 +1,6 @@
 // Copyright Epic Games, Inc. All Rights Reserved.
 
 import { DateTimeFormatOptions } from 'intl';
-import { Args } from '../common/args';
 import { Badge } from '../common/badge';
 import { Random, setDefault } from '../common/helper';
 import { ContextualLogger } from '../common/logger';
@@ -14,7 +13,6 @@ import { PersistedSlackMessage, PersistedSlackMessages } from './persistedmessag
 import { RoboArgs } from './roboargs';
 import { BlockageNodeOpUrls } from './roboserver';
 import { Slack, SlackAttachment, SlackFile, SlackLinkButtonAction, SlackLinkButtonsAttachment, SlackMessageField, SlackMessage, SlackMessageStyles } from './slack';
-import { Vault } from './vault'
 import { WebServer } from '../common/webserver'
 import { DummySlackApp } from '../common/dummyslackserver'
 
@@ -28,23 +26,19 @@ const KNOWN_BOT_EMAILS = [[REDACTED_EMAIL]'];
 let SLACK_TOKENS: {[name: string]: string} = {}
 const SLACK_DEV_DUMMY_TOKEN = 'dev'
 
-let args: Args
-export function notificationsInit(inArgs: Args, vault: Vault) {
-	
-	args = inArgs
-
-	if (args.devMode) {
+export function notificationsInit() {
+	if (RoboArgs.args.devMode) {
 		Badge.setDevMode()
 	}
 
-	if (args.slackDomain.indexOf('localhost') >= 0) { // todo should parse out port
+	if (RoboArgs.args.slackDomain.indexOf('localhost') >= 0) { // todo should parse out port
 		const slackServer = new WebServer(new ContextualLogger('dummy Slack'))
 		slackServer.addApp(DummySlackApp)
 		slackServer.open(8811, 'http')
 		return
 	}
 
-	const tokensObj = vault.slackTokens
+	const tokensObj = RoboArgs.vault.slackTokens
 	if (tokensObj) {
 		SLACK_TOKENS = tokensObj
 	}
@@ -55,14 +49,14 @@ export function isUserAKnownBot(user: string) {
 }
 
 export async function postToRobomergeAlerts(message: string, replies?: (string|SlackFile|SlackMessage)[]) {
-	if (!args.devMode) {
-		return postMessageToChannel(`${message}\n\n${RoboArgs.externalUrl}`, args.slackAlertChannel, SlackMessageStyles.DANGER, replies)
+	if (!RoboArgs.args.devMode) {
+		return postMessageToChannel(`${message}\n\n${RoboArgs.externalUrl}`, RoboArgs.args.slackAlertChannel, SlackMessageStyles.DANGER, replies)
 	}
 }
 
 export async function postMessageToChannel(message: string, channel: string, style: SlackMessageStyles = SlackMessageStyles.GOOD, replies?: (string|SlackFile|SlackMessage)[]) {
 	if (SLACK_TOKENS.bot) {
-		const slack = new Slack({id: channel, botToken: SLACK_TOKENS.bot, userToken: SLACK_TOKENS.user}, args.slackDomain, new ContextualLogger(`Post Message to ${channel}`))	
+		const slack = new Slack({id: channel, botToken: SLACK_TOKENS.bot, userToken: SLACK_TOKENS.user}, RoboArgs.args.slackDomain, new ContextualLogger(`Post Message to ${channel}`))	
 		const thread_ts = await slack.postMessage({text: message, style, channel, mrkdwn: true })
 		for (const reply of (replies || [])) {
 			if (typeof reply === 'string') {
@@ -999,7 +993,7 @@ export class BotNotifications implements BotEventHandler {
 		// Start collecting our attachments
 		let attachCollection : SlackAttachment[] = []
 		
-		const helpURL = args.helpPageURL != '/help' ? args.helpPageURL : `${this.externalRobomergeUrl}/help`
+		const helpURL = RoboArgs.args.helpPageURL != '/help' ? RoboArgs.args.helpPageURL : `${this.externalRobomergeUrl}/help`
 
 		// General information
 		attachCollection.push({
@@ -1234,11 +1228,11 @@ export function bindBotNotifications(events: BotEvents, slackChannelOverrides: [
 
 	let slackMessages
 
-	const botToken = (!args.devMode || args.useSlackInDev) && SLACK_TOKENS.bot || args.devMode && args.useSlackInDev && SLACK_DEV_DUMMY_TOKEN 
-	const userToken = (!args.devMode || args.useSlackInDev) && SLACK_TOKENS.user || args.devMode && args.useSlackInDev && SLACK_DEV_DUMMY_TOKEN
+	const botToken = (!RoboArgs.args.devMode || RoboArgs.args.useSlackInDev) && SLACK_TOKENS.bot || RoboArgs.args.devMode && RoboArgs.args.useSlackInDev && SLACK_DEV_DUMMY_TOKEN 
+	const userToken = (!RoboArgs.args.devMode || RoboArgs.args.useSlackInDev) && SLACK_TOKENS.user || RoboArgs.args.devMode && RoboArgs.args.useSlackInDev && SLACK_DEV_DUMMY_TOKEN
 	if (botToken && events.botConfig.slackChannel) {
 		logger.info('Enabling Slack messages for ' +  events.botname)
-		slackMessages = new SlackMessages(new Slack({id: events.botConfig.slackChannel, botToken, userToken}, args.slackDomain, logger), persistence, logger)
+		slackMessages = new SlackMessages(new Slack({id: events.botConfig.slackChannel, botToken, userToken}, RoboArgs.args.slackDomain, logger), persistence, logger)
 	}
 		
 	events.registerHandler(new BotNotifications(events.botname, events.botConfig.slackChannel, externalUrl, blockageUrlGenerator, logger, slackMessages, slackChannelOverrides))
diff --git a/Engine/Extras/RoboMerge/v3/src/robo/robo.ts b/Engine/Extras/RoboMerge/v3/src/robo/robo.ts
--- a/Engine/Extras/RoboMerge/v3/src/robo/robo.ts
+++ b/Engine/Extras/RoboMerge/v3/src/robo/robo.ts
@@ -25,7 +25,6 @@ import { RoboServer } from './roboserver';
 import { notificationsInit } from './notifications'
 import { RoboArgs } from './roboargs'
 import { settingsInit } from './settings'
-import { Vault } from './vault'
 import { doDBPersistedMessagesHousekeeping } from './persistedmessages';
 
 /*************************
@@ -285,9 +284,8 @@ if (!maybeNullArgs) {
 const args = maybeNullArgs!
 RoboArgs.Init(args)
 
-const vault = new Vault(args.vault)
-notificationsInit(args, vault)
-settingsInit(args)
+notificationsInit()
+settingsInit()
 
 /**
  * Connect to MongoDB.
@@ -655,7 +653,7 @@ async function _initGraphBots(logger: ContextualLogger) {
 async function main(logger: ContextualLogger) {
 	while (true) {
 		try {
-			await initializePerforce(logger, vault, args.devMode);
+			await initializePerforce(logger, RoboArgs.vault, args.devMode);
 			break;
 		}
 		catch (err) {
diff --git a/Engine/Extras/RoboMerge/v3/src/robo/roboargs.ts b/Engine/Extras/RoboMerge/v3/src/robo/roboargs.ts
--- a/Engine/Extras/RoboMerge/v3/src/robo/roboargs.ts
+++ b/Engine/Extras/RoboMerge/v3/src/robo/roboargs.ts
@@ -1,13 +1,24 @@
 // Copyright Epic Games, Inc. All Rights Reserved.
 
 import { Args } from '../common/args';
+import { Vault } from './vault';
 
 export class RoboArgs {
 
 	private static _args: Args
-	
+	private static _vault: Vault
+
 	static Init(args: Args) {
 		this._args = args
+		this._vault = new Vault(args.vault)
+	}
+
+	static get args() {
+		return this._args
+	}
+
+	static get vault() {
+		return this._vault
 	}
 
 	static get exclusiveLockOpenedsToRun() {
diff --git a/Engine/Extras/RoboMerge/v3/src/robo/settings.ts b/Engine/Extras/RoboMerge/v3/src/robo/settings.ts
--- a/Engine/Extras/RoboMerge/v3/src/robo/settings.ts
+++ b/Engine/Extras/RoboMerge/v3/src/robo/settings.ts
@@ -1,7 +1,6 @@
 // Copyright Epic Games, Inc. All Rights Reserved.
 
 import * as fs from 'fs';
-import { Args } from '../common/args';
 import { ContextualLogger } from '../common/logger';
 import { PerforceContext, Workspace } from '../common/perforce';
 import { VersionReader } from '../common/version';
@@ -25,9 +24,7 @@ const BotModel = require('./models/bot');
  */
 const VERSION = VersionReader.getPackageVersion()
 
-let args: Args
-export function settingsInit(inArgs: Args) {
-	args = inArgs
+export function settingsInit() {
 }
 	
 interface SaveObjectOpts {
@@ -104,7 +101,7 @@ const FilePurposeExtensionMap = {
 }
 
 function getPersistenceFilepath(botname: string) {
-	return args.persistenceDir + `/${botname}.settings`
+	return RoboArgs.args.persistenceDir + `/${botname}.settings`
 }
 
 type SavingState = 'Saving' | 'NeedsSave'
@@ -123,8 +120,8 @@ export class Settings {
 
 		this.lastPersistBackupTime = new Date(Date.now())
 		this.botname = botname.toLowerCase() 
-		if (args.persistenceBackupFrequency > 0) {
-			this.p4backupWorkspace = {directory: fs.realpathSync(args.persistenceDir), name: args.persistenceBackupWorkspace}
+		if (RoboArgs.args.persistenceBackupFrequency > 0) {
+			this.p4backupWorkspace = {directory: fs.realpathSync(RoboArgs.args.persistenceDir), name: RoboArgs.args.persistenceBackupWorkspace}
 		}
 
 		// see if we should enable saves
@@ -327,7 +324,7 @@ export class Settings {
 							}
 						})
 					}
-					else if ((Date.now() - this.lastPersistBackupTime.getTime()) / 60000 >= args.persistenceBackupFrequency) {
+					else if ((Date.now() - this.lastPersistBackupTime.getTime()) / 60000 >= RoboArgs.args.persistenceBackupFrequency) {
 						this.lastPersistBackupTime = new Date(Date.now())
 
 						this.p4.revertFile(this.p4backupWorkspace, persistentFile, ['-k'])
