Interactive Slack buttons

diff --git a/Engine/Extras/RoboMerge/v3/src/common/webserver.ts b/Engine/Extras/RoboMerge/v3/src/common/webserver.ts
--- a/Engine/Extras/RoboMerge/v3/src/common/webserver.ts
+++ b/Engine/Extras/RoboMerge/v3/src/common/webserver.ts
@@ -68,6 +68,7 @@ export interface WebRequest {
 	url: URL,
 	cookies: string
 	reqData?: string
+	headers?: Record<string, string | string[] | undefined>
 }
 
 interface AppConstructor {
@@ -381,7 +382,7 @@ export class WebServer {
 			const match = url.pathname.match(handler.route)
 			if (match) {
 				try {
-					await handler.handler(response, {url: url, cookies: req.headers.cookie as string, reqData}, match)
+					await handler.handler(response, {url: url, cookies: req.headers.cookie as string, reqData, headers: req.headers as Record<string, string | string[] | undefined>}, match)
 				}
 				catch (err) {
 					if (match[0] === '/api/crashme') {
diff --git a/Engine/Extras/RoboMerge/v3/src/robo/graphbot.ts b/Engine/Extras/RoboMerge/v3/src/robo/graphbot.ts
--- a/Engine/Extras/RoboMerge/v3/src/robo/graphbot.ts
+++ b/Engine/Extras/RoboMerge/v3/src/robo/graphbot.ts
@@ -114,12 +114,13 @@ export class GraphBot implements BotEventHandler {
 	initBots(ubergraph: GraphAPI, previewMode: boolean) {
 		this.eventTriggers = new BotEventTriggers(this.branchGraph.botname, this.branchGraph.config)
 		this.eventTriggers.registerHandler(this)
-		const blockageUrlGenerator: NodeOpUrlGenerator = (blockage : Blockage | null) => { 
+		const hasSocketMode = RoboArgs.hasSocketMode
+		const blockageUrlGenerator: NodeOpUrlGenerator = (blockage : Blockage | null) => {
 			if (!blockage) {
 				return null
 			}
 			const sourceNode = this.findNode(blockage.change.branch)
-			return sourceNode ? sourceNode.getBlockageUrls(blockage) : null
+			return sourceNode ? sourceNode.getBlockageUrls(blockage, hasSocketMode) : null
 		}
 
 		// bind handlers to bot events
diff --git a/Engine/Extras/RoboMerge/v3/src/robo/models/slackmessages.ts b/Engine/Extras/RoboMerge/v3/src/robo/models/slackmessages.ts
--- a/Engine/Extras/RoboMerge/v3/src/robo/models/slackmessages.ts
+++ b/Engine/Extras/RoboMerge/v3/src/robo/models/slackmessages.ts
@@ -13,7 +13,8 @@ var slackMessagesSchema = new mongoose.Schema({
   timestamp: { type: String },
   timestamp_forSunset: { type: Number },
   permalink: { type: String },
-  message: {  }
+  message: {  },
+  buttonReplyTimestamp: { type: String }
 });
 
 module.exports = mongoose.model('SlackMessages', slackMessagesSchema);
diff --git a/Engine/Extras/RoboMerge/v3/src/robo/nodebot.ts b/Engine/Extras/RoboMerge/v3/src/robo/nodebot.ts
--- a/Engine/Extras/RoboMerge/v3/src/robo/nodebot.ts
+++ b/Engine/Extras/RoboMerge/v3/src/robo/nodebot.ts
@@ -1376,7 +1376,7 @@ export class NodeBot extends PerforceStatefulBot implements NodeBotInterface {
 		return false
 	}
 
-	getBlockageUrls(blockage: Blockage, skipReason?: string) {
+	getBlockageUrls(blockage: Blockage, hasSocketMode = false, skipReason?: string) {
 		const externalUrl = this.externalRobomergeUrl
 		const botname = this.graphBotName
 		const branchname = this.branch.name
@@ -1384,20 +1384,30 @@ export class NodeBot extends PerforceStatefulBot implements NodeBotInterface {
 		const disallowSkip = this.disallowSkipForBlockage(blockage)
 
 		// Call static method so we can centralize the logic (see Notifications.sendTestMessage())
-		return NodeBot.getBlockageUrls(blockage, externalUrl, botname, branchname, cl, disallowSkip, skipReason)
+		return NodeBot.getBlockageUrls(blockage, externalUrl, botname, branchname, cl, disallowSkip, hasSocketMode, skipReason)
 	}
 
-	static getBlockageUrls(blockage: Blockage, externalUrl: string, botname: string, branchname: string, cl: string, disallowSkip: boolean, skipReason?: string) {
+	static getBlockageUrls(blockage: Blockage, externalUrl: string, botname: string, branchname: string, cl: string, disallowSkip: boolean, hasSocketMode = false, skipReason?: string) {
+		const actionBranchName = blockage.action ? blockage.action.branch.name : undefined
 		let urls : BlockageNodeOpUrls = {
 			acknowledgeUrl: OperationUrlHelper.createAcknowledgeUrl(
 				externalUrl,
 				botname,
 				branchname,
 				cl,
-				blockage.action ? blockage.action.branch.name : undefined
-			)
+				actionBranchName
+			),
+			params: hasSocketMode ? {
+				bot: botname,
+				branch: branchname,
+				cl: parseInt(cl, 10),
+				host: externalUrl,
+				edge: actionBranchName,
+				target: actionBranchName,
+				targetStream: blockage.action ? blockage.action.branch.stream : undefined
+			} : undefined
 		}
-		
+
 		if (blockage.action) {
 			// Do not generate create shelf or stomp urls for exclusive checkouts
 			if (blockage.failure.kind === 'Merge conflict') {
diff --git a/Engine/Extras/RoboMerge/v3/src/robo/notifications.ts b/Engine/Extras/RoboMerge/v3/src/robo/notifications.ts
--- a/Engine/Extras/RoboMerge/v3/src/robo/notifications.ts
+++ b/Engine/Extras/RoboMerge/v3/src/robo/notifications.ts
@@ -11,8 +11,8 @@ import { BlockageStatusChangedEvent, BotEventHandler, BotEvents } from './events
 import { NodeBot } from './nodebot';
 import { PersistedSlackMessage, PersistedSlackMessages } from './persistedmessages';
 import { RoboArgs } from './roboargs';
-import { BlockageNodeOpUrls } from './roboserver';
-import { Slack, SlackAttachment, SlackFile, SlackLinkButtonAction, SlackLinkButtonsAttachment, SlackMessageField, SlackMessage, SlackMessageStyles } from './slack';
+import { BlockageNodeOpParams, BlockageNodeOpUrls } from './roboserver';
+import { Slack, SlackActionsBlock, SlackAttachment, SlackFile, SlackInteractiveButton, SlackLinkButtonsAttachment, SlackMessageField, SlackMessage, SlackMessageStyles } from './slack';
 import { WebServer } from '../common/webserver'
 import { DummySlackApp } from '../common/dummyslackserver'
 
@@ -222,8 +222,8 @@ export class SlackMessages {
 			message.target = resolveBranchArg(branchArg, true)
 			message.cl = cl
 
-			let timestamp
-			let permalink
+			let timestamp: string
+			let permalink: string
 			try {
 				timestamp = await this.slack.postMessage(message)
 				permalink = await this.slack.getPermalink(timestamp, message.channel)
@@ -242,20 +242,72 @@ export class SlackMessages {
 		}
 	}
 
-	async postReply(cl: number, branchArg: BranchArg, message: SlackMessage) {
+	async postReply(cl: number, branchArg: BranchArg, message: SlackMessage): Promise<string | undefined> {
 		const findResult = await this.find(cl, branchArg, message.channel)
 
 		if (findResult) {
 			try {
-				await this.slack.reply(findResult.timestamp, message)
+				return await this.slack.reply(findResult.timestamp, message)
 			}
 			catch (err) {
 				console.error('Error updating message in Slack! ' + err.toString())
-				return
+				return undefined
 			}
 		}
 		else {
 			this.smLogger.error(`Failed to find message record for ${cl}:${branchArg}:${message.channel}`)
+			return undefined
+		}
+	}
+
+	async postButtonReply(cl: number, branchArg: BranchArg, message: SlackMessage) {
+		const persistedMsg = await this.find(cl, branchArg, message.channel)
+		if (!persistedMsg) {
+			this.smLogger.error(`Failed to find message record for ${cl}:${branchArg}:${message.channel}`)
+				return
+			}
+		try {
+			const ts = await this.slack.reply(persistedMsg.timestamp, message)
+			persistedMsg.buttonReplyTimestamp = ts
+			persistedMsg.update()
+		}
+		catch (err) {
+			this.smLogger.error(`Error posting button reply to Slack: ${err}`)
+		}
+	}
+
+	async updateButtonReply(cl: number, branchArg: BranchArg, message: SlackMessage) {
+		const persistedMsg = await this.find(cl, branchArg, message.channel)
+		if (!persistedMsg) {
+			this.smLogger.error(`Failed to find message record for ${cl}:${branchArg}:${message.channel}`)
+			return
+		}
+		if (persistedMsg.buttonReplyTimestamp) {
+			try {
+				await this.slack.update(persistedMsg.buttonReplyTimestamp, message)
+			}
+			catch (err) {
+				this.smLogger.error(`Error updating button reply in Slack: ${err}`)
+			}
+		}
+		else {
+			await this.postButtonReply(cl, branchArg, message)
+		}
+	}
+
+	async clearButtonReplies(cl: number, branchArg: BranchArg) {
+		const messages = await this.persistedMessages.findAll(cl, branchArg)
+		for (const msg of messages) {
+			if (msg.buttonReplyTimestamp) {
+				try {
+					await this.slack.delete(msg.buttonReplyTimestamp, msg.messageOpts.channel)
+				}
+				catch (err) {
+					this.smLogger.error(`Error deleting button reply in Slack: ${err}`)
+				}
+				msg.buttonReplyTimestamp = undefined
+				msg.update()
+			}
 		}
 	}
 
@@ -417,6 +469,13 @@ export class SlackMessages {
 	getSlackUser(emailAddress: string) {
 		return this.slack.lookupUserIdByEmail(emailAddress.toLowerCase())
 	}
+
+	async getSlackUserEmail(slackUserId: string): Promise<string | null> {
+		const result = await this.slack.getUserInfo(slackUserId)
+		return (result && result.ok && result.user && result.user.profile && result.user.profile.email)
+			? result.user.profile.email as string
+			: null
+	}
 }
 
 export function makeClLink(cl: number, alias?: string) {
@@ -551,6 +610,9 @@ export class BotNotifications implements BotEventHandler {
 			text = `${channelPing}, please resolve the following ${issue}:`
 		}														
 
+		// Look up blockage URLs up front — needed for both channel message buttons and DM
+		const blockageUrls = targetBranch ? this.blockageUrlGenerator(blockage) : null
+
 		let message = this.makeSlackChannelMessage(
 			`${sourceBranch.displayName} blocked! (${issue})`,
 			text,
@@ -606,41 +668,90 @@ export class BotNotifications implements BotEventHandler {
 			}
 		}
 
-		// For first-time blockages, post supplementary details (conflict files, locked file pings)
-		if (isNew) {
-			await this._postBlockageDetails(blockage, changeInfo.cl, branchArg, channelsToPostTo)
-		}
+		if (targetBranch) {
+			// For first-time blockages, post supplementary details (conflict files, locked file pings)
+			if (isNew) {
+				await this._postBlockageDetails(blockage, changeInfo.cl, branchArg, channelsToPostTo)
+			}
 
-		// Post message to owner in DM
-		if (!isBotUser && targetBranch && userEmail) {
-			let dm: SlackMessage
-			if (blockage.approval) {
-				dm = {
-					title: 'Approval needed to commit to ' + targetBranch.displayName,
-					text: `Your change has been shelved in ${makeClLink(blockage.approval.shelfCl)} and sent to <#${blockage.approval.settings.channelId}> for approval\n\n` +
-							blockage.approval.settings.description,
-					channel: "",
-					mrkdwn: true
+			// Post action buttons as a thread reply (new blockage) or update existing (re-trigger).
+			// updateButtonReply falls back to posting if no existing ts is tracked.
+			if (blockageUrls?.params) {
+				const actionsBlock = this._makeActionsBlock(blockageUrls, targetBranch.displayName)
+				if (actionsBlock.elements.length > 0) {
+					for (const channel of channelsToPostTo) {
+						if (isNew) {
+							await this.slackMessages.postButtonReply(changeInfo.cl, branchArg, {
+								text: '',
+								channel,
+								mrkdwn: false,
+								blocks: [actionsBlock]
+							})
+						} 
+						else {
+							await this.slackMessages.updateButtonReply(changeInfo.cl, branchArg, {
+								text: '',
+								channel,
+								mrkdwn: false,
+								blocks: [actionsBlock]
+							})
+						}
+					}
 				}
 			}
-			else {
-				const dmText = `Your change (${makeClLink(changeInfo.source_cl)}) ` +
-					`hit '${issue}' while merging from *${sourceBranch.displayName}* to *${targetBranch.displayName}*.\n\n` +
-					'`' + blockage.change.description.substring(0, 80) + '`\n\n' +
-					"*_Resolving this blockage is time sensitive._*"
+
+			// Post message to owner in DM
+			if (!isBotUser && userEmail) {
+				let dm: SlackMessage
+				if (blockage.approval) {
+					dm = {
+						title: 'Approval needed to commit to ' + targetBranch.displayName,
+						text: `Your change has been shelved in ${makeClLink(blockage.approval.shelfCl)} and sent to <#${blockage.approval.settings.channelId}> for approval\n\n` +
+								blockage.approval.settings.description,
+						channel: "",
+						mrkdwn: true
+					}
+				}
+				else {
+					if (!blockageUrls) {
+						const error = `Could not get blockage URLs for blockage -- CL ${blockage.change.cl}`
+						this.botNotificationsLogger.printException(error)
+						throw error
+					}
+
+					const dmText = `Your change (${makeClLink(changeInfo.source_cl)}) ` +
+						`hit '${issue}' while merging from *${sourceBranch.displayName}* to *${targetBranch.displayName}*.\n\n` +
+						'`' + blockage.change.description.substring(0, 80) + '`\n\n' +
+						"*_Resolving this blockage is time sensitive._*"
 				
-				const urls = this.blockageUrlGenerator(blockage)
-				if (!urls) {
-					const error = `Could not get blockage URLs for blockage -- CL ${blockage.change.cl}`
-					this.botNotificationsLogger.printException(error)
-					throw error
+					dm = this.makeSlackDirectMessage(dmText, changeInfo.cl, changeInfo.cl, targetBranch.displayName, blockageUrls, threadLinks.filter((l): l is string => !!l))
 				}
 
-				dm = this.makeSlackDirectMessage(dmText, changeInfo.cl, changeInfo.cl, targetBranch.displayName, urls, threadLinks)
+				this.slackMessages.postDM(userEmail, changeInfo.cl, targetBranch, dm)
 			}
+		}
+	}
+
+	private _makeActionsBlock(conflictUrls: BlockageNodeOpUrls, targetBranch: string): SlackActionsBlock {
+		const params = conflictUrls.params!
+		const buttons: SlackInteractiveButton[] = []
 
-			this.slackMessages.postDM(userEmail, changeInfo.cl, targetBranch, dm)
+		buttons.push(this.makeInteractiveButton('Acknowledge Conflict', 'acknowledge', params, 'primary'))
+
+		if (conflictUrls.createShelfUrl) {
+			buttons.push(this.makeInteractiveButton(`Create Shelf in ${targetBranch}`, 'create_shelf', params, 'primary'))
+		}
+		if (conflictUrls.skipUrl) {
+			buttons.push(this.makeInteractiveButton(`Skip Merge to ${targetBranch}`, 'skip', params))
+		}
+		if (conflictUrls.stompUrl) {
+			buttons.push(this.makeInteractiveButton(`Stomp Changes in ${targetBranch}`, 'stompchanges', params, 'danger'))
+		}
+		if (conflictUrls.unlockUrl) {
+			buttons.push(this.makeInteractiveButton(`Unlock ${targetBranch}`, 'unlockchanges', params, 'danger'))
 		}
+
+		return { type: 'actions', elements: buttons }
 	}
 
 	/** Post supplementary details for a blockage: conflict file list and/or lock holder pings.
@@ -746,6 +857,22 @@ export class BotNotifications implements BotEventHandler {
 
 		// Post the full supplementary details for the new failure kind
 		await this._postBlockageDetails(newBlockage, cl, branchArg, channels)
+
+		// Update the action buttons to reflect the new blockage state
+		const blockageUrls = targetBranch ? this.blockageUrlGenerator(newBlockage) : null
+		if (targetBranch && blockageUrls?.params) {
+			const actionsBlock = this._makeActionsBlock(blockageUrls, targetBranch.displayName)
+			if (actionsBlock.elements.length > 0) {
+				for (const channel of channels) {
+					await this.slackMessages.updateButtonReply(cl, branchArg, {
+						text: '',
+						channel,
+						mrkdwn: false,
+						blocks: [actionsBlock]
+					})
+				}
+			}
+		}
 	}
 
 	onBlockageAcknowledged(info: PersistentConflict) {
@@ -850,6 +977,7 @@ export class BotNotifications implements BotEventHandler {
 					this.slackMessages!.postOrUpdate(info.cl, targetKey, message)
 				}
 			})
+			this.slackMessages.clearButtonReplies(info.cl, targetKey)
 		}
 	}
 
@@ -988,7 +1116,6 @@ export class BotNotifications implements BotEventHandler {
 	}
 
 	// This is an extremely opinionated function to send a stylized direct message to the end user.
-	//makeSlackChannelMessageOpts(title: string, style: string, clDesc: string, sourceBranch: BranchArg,targetBranch?: BranchArg, author?: string, buttons?: SlackLinkButtonsAttachment[]) 
 	private makeSlackDirectMessage(messageText: string, sourceCl: number, conflictCl: number, targetBranch: string, conflictUrls: BlockageNodeOpUrls, threadLinks?: string[]) : SlackMessage {
 		// Start collecting our attachments
 		let attachCollection : SlackAttachment[] = []
@@ -1022,88 +1149,106 @@ export class BotNotifications implements BotEventHandler {
 			mrkdwn_in: ["text"]
 		})
 
-		// Acknowledge button
 		const conflictClLink = makeClLink(conflictCl, 'conflict CL #' + conflictCl)
-		attachCollection.push(<SlackLinkButtonsAttachment>{
-			text: `"I will merge ${conflictClLink} to *${targetBranch}* myself."`,
-			fallback: `Please acknowledge blockages at ${this.externalRobomergeUrl}`,
-			mrkdwn_in: ["text"],
-			actions: [this.generateAcknowledgeButton("Acknowledge Conflict", conflictUrls.acknowledgeUrl)]
-		})
+		const params = conflictUrls.params
+		if (params) {
+			// Block Kit interactive buttons — actions performed directly in Slack
+			const buttons: SlackInteractiveButton[] = []
 
-		// Create shelf button
-		if (conflictUrls.createShelfUrl) {
-			attachCollection.push(<SlackLinkButtonsAttachment>{
-				text: `"Please create a shelf with the conflicts encountered while merging ${makeClLink(sourceCl)} into *${targetBranch}*"`,
-				fallback: `You can create a shelf at ${this.externalRobomergeUrl}`,
-				mrkdwn_in: ["text"],
-				actions: [this.generateCreateShelfButton(`Create Shelf in ${targetBranch}`, conflictUrls.createShelfUrl)]
-			})
-		}
+			// Acknowledge
+			buttons.push(this.makeInteractiveButton('Acknowledge Conflict', 'acknowledge', params, 'primary'))
 
-		// Skip button
-		if (conflictUrls.skipUrl) {
-			attachCollection.push(<SlackLinkButtonsAttachment>{
-				text: `"${makeClLink(sourceCl)} should not be automatically merged to *${targetBranch}*."`,
-				fallback: `You can skip work at ${this.externalRobomergeUrl}`,
-				mrkdwn_in: ["text"],
-				actions: [this.generateSkipButton(`Skip Merge to ${targetBranch}`, conflictUrls.skipUrl)]
+			// Create Shelf (workspace selected via modal)
+			if (conflictUrls.createShelfUrl) {
+				buttons.push(this.makeInteractiveButton(`Create Shelf in ${targetBranch}`, 'create_shelf', params, 'primary'))
+			}
+
+			// Skip
+			if (conflictUrls.skipUrl) {
+				buttons.push(this.makeInteractiveButton(`Skip Merge to ${targetBranch}`, 'skip', params))
+			}
+
+			// Stomp
+			if (conflictUrls.stompUrl) {
+				buttons.push(this.makeInteractiveButton(`Stomp Changes in ${targetBranch}`, 'stompchanges', params, 'danger'))
+			}
+
+			// Unlock (exclusive check-out)
+			if (conflictUrls.unlockUrl) {
+				buttons.push(this.makeInteractiveButton(`Unlock ${targetBranch}`, 'unlockchanges', params, 'danger'))
+			}
+
+			const actionsBlock: SlackActionsBlock = { type: 'actions', elements: buttons }
+
+			attachCollection.push({
+				text: `"I will merge ${conflictClLink} to *${targetBranch}* myself."`,
+				mrkdwn_in: ["text"]
 			})
-		}
 
-		// Create stomp button if this isn't an exclusive checkout
-		if (conflictUrls.stompUrl) {
+			return {
+				text: messageText,
+				mrkdwn: true,
+				channel: "",
+				attachments: attachCollection,
+				blocks: [actionsBlock]
+			}
+		}
+		else {
+			// Fallback: legacy link buttons (no params means no interactive support)
 			attachCollection.push(<SlackLinkButtonsAttachment>{
-				text: `"The changes in ${makeClLink(sourceCl)} should stomp the work in *${targetBranch}*."`,
-				fallback: `You can stomp work at ${this.externalRobomergeUrl}`,
+				text: `"I will merge ${conflictClLink} to *${targetBranch}* myself."`,
+				fallback: `Please acknowledge blockages at ${this.externalRobomergeUrl}`,
 				mrkdwn_in: ["text"],
-				actions: [this.generateStompButton(`Stomp Changes in ${targetBranch}`, conflictUrls.stompUrl)]
+				actions: [{ type: "button", text: "Acknowledge Conflict", url: conflictUrls.acknowledgeUrl, style: "primary" }]
 			})
-		}
 
-		// Return SlackMessage for our direct message
-		return {
-			text: messageText,
-			mrkdwn: true,
-			channel: "",
-			attachments: attachCollection
-		}
-	}
+			if (conflictUrls.createShelfUrl) {
+				attachCollection.push(<SlackLinkButtonsAttachment>{
+					text: `"Please create a shelf with the conflicts encountered while merging ${makeClLink(sourceCl)} into *${targetBranch}*"`,
+					fallback: `You can create a shelf at ${this.externalRobomergeUrl}`,
+					mrkdwn_in: ["text"],
+					actions: [{ type: "button", text: `Create Shelf in ${targetBranch}`, url: conflictUrls.createShelfUrl, style: "primary" }]
+				})
+			}
 
-	private generateAcknowledgeButton(buttonText: string, url: string) : SlackLinkButtonAction {
-		return { 
-			type: "button",
-			text: buttonText,
-			url,
-			style: "primary"
-		}
-	}
+			if (conflictUrls.skipUrl) {
+				attachCollection.push(<SlackLinkButtonsAttachment>{
+					text: `"${makeClLink(sourceCl)} should not be automatically merged to *${targetBranch}*."`,
+					fallback: `You can skip work at ${this.externalRobomergeUrl}`,
+					mrkdwn_in: ["text"],
+					actions: [{ type: "button", text: `Skip Merge to ${targetBranch}`, url: conflictUrls.skipUrl, style: "default" }]
+				})
+			}
 
-	private generateCreateShelfButton(buttonText: string, url: string) : SlackLinkButtonAction {
-		return { 
-			type: "button",
-			text: buttonText,
-			url,
-			style: "primary"
-		}
-	}
+			if (conflictUrls.stompUrl) {
+				attachCollection.push(<SlackLinkButtonsAttachment>{
+					text: `"The changes in ${makeClLink(sourceCl)} should stomp the work in *${targetBranch}*."`,
+					fallback: `You can stomp work at ${this.externalRobomergeUrl}`,
+					mrkdwn_in: ["text"],
+					actions: [{ type: "button", text: `Stomp Changes in ${targetBranch}`, url: conflictUrls.stompUrl, style: "danger" }]
+				})
+			}
 
-	private generateSkipButton(buttonText: string, url: string) : SlackLinkButtonAction {
-		return { 
-			type: "button",
-			text: buttonText,
-			url,
-			style: 'default'
+			return {
+				text: messageText,
+				mrkdwn: true,
+				channel: "",
+				attachments: attachCollection
+			}
 		}
 	}
 
-	private generateStompButton(buttonText: string, url: string) : SlackLinkButtonAction {
-		return { 
-			type: "button",
-			text: buttonText,
-			url,
-			style: "danger"
+	private makeInteractiveButton(buttonText: string, actionId: string, params: BlockageNodeOpParams, style?: 'primary' | 'danger'): SlackInteractiveButton {
+		const btn: SlackInteractiveButton = {
+			type: 'button',
+			text: { type: 'plain_text', text: buttonText },
+			action_id: actionId,
+			value: JSON.stringify(params)
+		}
+		if (style) {
+			btn.style = style
 		}
+		return btn
 	}
 
 	
@@ -1150,9 +1295,10 @@ export class BotNotifications implements BotEventHandler {
 				message.fields = newFields
 			}
 
-			// Delete button attachments sent via Robomerge Slack App
+			// Delete button attachments and blocks sent via Robomerge Slack App
 			if (message.attachments) {
 				delete message.attachments
+				message.blocks = []
 
 				// Hacky: If we remove attachments, we'll no longer have a link to the CL in the message.
 				// UE-72320 - Add in a link to the original changelist
diff --git a/Engine/Extras/RoboMerge/v3/src/robo/persistedmessages.ts b/Engine/Extras/RoboMerge/v3/src/robo/persistedmessages.ts
--- a/Engine/Extras/RoboMerge/v3/src/robo/persistedmessages.ts
+++ b/Engine/Extras/RoboMerge/v3/src/robo/persistedmessages.ts
@@ -16,6 +16,7 @@ type PersistedMessagePayload = {
 	timestamp: string
 	permalink: string
 	messageOpts: SlackMessage
+	buttonReplyTimestamp?: string
 }
 
 type PersistedMessagePersistence = {[key: string]: PersistedMessagePayload}
@@ -68,11 +69,13 @@ export class PersistedSlackMessage implements PersistedMessagePayload {
 	timestamp: string
 	permalink: string
 	messageOpts: SlackMessage
+	buttonReplyTimestamp?: string
 
 	constructor(payload: PersistedMessagePayload, private _key: PersistedMessageKey, private persistedMessages: PersistedSlackMessages) {
 		this.timestamp = payload.timestamp
 		this.permalink = payload.permalink
 		this.messageOpts = payload.messageOpts
+		this.buttonReplyTimestamp = payload.buttonReplyTimestamp
 	}
 
 	get key() {
@@ -120,10 +123,11 @@ class DBPersistedSlackMessages extends PersistedSlackMessages {
 		return SlackMessagesModel.collection.updateOne(
 			index,
 			{$set:{
-				...index, 
-				timestamp: message.timestamp, 
-				permalink: message.permalink, 
-				message: message.messageOpts, 
+				...index,
+				timestamp: message.timestamp,
+				permalink: message.permalink,
+				message: message.messageOpts,
+				buttonReplyTimestamp: message.buttonReplyTimestamp,
 				timestamp_forSunset: parseFloat(message.timestamp)
 			}},
 			{upsert:true}
@@ -150,9 +154,10 @@ class DBPersistedSlackMessages extends PersistedSlackMessages {
 		const messageRecord = await SlackMessagesModel.find(index)
 		if (messageRecord.length == 1) {
 			const persistedMessage = new PersistedSlackMessage({
-						messageOpts: messageRecord[0].message, 
+						messageOpts: messageRecord[0].message,
 						permalink: messageRecord[0].permalink,
-						timestamp: messageRecord[0].timestamp
+						timestamp: messageRecord[0].timestamp,
+						buttonReplyTimestamp: messageRecord[0].buttonReplyTimestamp
 					}, key, this)
 			this.messageCache[cacheKey] = persistedMessage
 			return persistedMessage
@@ -178,9 +183,10 @@ class DBPersistedSlackMessages extends PersistedSlackMessages {
 			const cacheKey = key.toString()
 			if (!this.messageCache[cacheKey]) {
 				const persistedMessage = new PersistedSlackMessage({
-						messageOpts: messageRecord.message, 
+						messageOpts: messageRecord.message,
 						permalink: messageRecord.permalink,
-						timestamp: messageRecord.timestamp
+						timestamp: messageRecord.timestamp,
+						buttonReplyTimestamp: messageRecord.buttonReplyTimestamp
 					}, key, this)
 				this.messageCache[cacheKey] = persistedMessage
 			}
diff --git a/Engine/Extras/RoboMerge/v3/src/robo/robo.ts b/Engine/Extras/RoboMerge/v3/src/robo/robo.ts
--- a/Engine/Extras/RoboMerge/v3/src/robo/robo.ts
+++ b/Engine/Extras/RoboMerge/v3/src/robo/robo.ts
@@ -22,6 +22,7 @@ import { IPC, Message } from './ipc';
 import { NodeBot } from './nodebot';
 import { roboAnalytics, setGlobalAnalytics } from './roboanalytics';
 import { RoboServer } from './roboserver';
+import { SlackSocket } from './slack-socket';
 import { notificationsInit } from './notifications'
 import { RoboArgs } from './roboargs'
 import { settingsInit } from './settings'
@@ -575,6 +576,16 @@ if (!args.noIPC) {
 else {
 	const sendMessage = (name: string, args?: any[]) => ipc.handle({name:name, args:args}).catch(err => console.log(`IPC error! ${err.toString()}\n${err.stack}`));
 
+	if (RoboArgs.hasSocketMode) {
+		const socket = new SlackSocket(RoboArgs.vault.slackAppToken!, args.externalUrl, {
+			sendMessage,
+			botToken: RoboArgs.vault.slackTokens?.bot,
+			slackDomain: args.slackDomain,
+			logger: roboStartupLogger
+		})
+		socket.start()
+	}
+
 	const ws = new RoboServer(args.externalUrl, sendMessage, () => roboStartupLogger.getLogTail(),
 		() => roboStartupLogger.info('Received: getLastCrash'),
 		() => roboStartupLogger.info('Received: stopBot'),
diff --git a/Engine/Extras/RoboMerge/v3/src/robo/roboargs.ts b/Engine/Extras/RoboMerge/v3/src/robo/roboargs.ts
--- a/Engine/Extras/RoboMerge/v3/src/robo/roboargs.ts
+++ b/Engine/Extras/RoboMerge/v3/src/robo/roboargs.ts
@@ -52,4 +52,8 @@ export class RoboArgs {
 	static get persistSettingsToFile() {
 		return !this.useMongo || this._args.alwaysPersistToFile
 	}
+
+	static get hasSocketMode() {
+		return !!this._vault.slackAppToken
+	}
 }
\ No newline at end of file
diff --git a/Engine/Extras/RoboMerge/v3/src/robo/roboserver.ts b/Engine/Extras/RoboMerge/v3/src/robo/roboserver.ts
--- a/Engine/Extras/RoboMerge/v3/src/robo/roboserver.ts
+++ b/Engine/Extras/RoboMerge/v3/src/robo/roboserver.ts
@@ -864,12 +864,29 @@ class RoboWebApp implements AppInterface {
 	}
 }
 
+// Raw operation parameters for building Slack interactive button values
+export interface BlockageNodeOpParams {
+	bot: string
+	branch: string
+	cl: number
+	// The robomerge instance that should handle this action (its externalUrl)
+	host: string
+	// For acknowledge and skip (the edge/target branch name)
+	edge?: string
+	// For createShelf, stomp, unlock (the target branch name — same value as edge)
+	target?: string
+	// For createShelf (the target branch stream, used for workspace lookup)
+	targetStream?: string
+}
+
 export interface BlockageNodeOpUrls {
 	acknowledgeUrl: string
 	createShelfUrl?: string
 	skipUrl?: string
 	stompUrl?: string
 	unlockUrl?: string
+	// Raw params for Slack interactive buttons
+	params?: BlockageNodeOpParams
 }
 export class OperationUrlHelper {
 	static createAcknowledgeUrl(externalRobomergeUrl: string, botname: string, branchname: string, changelistNum: string, edge?: string) {
diff --git a/Engine/Extras/RoboMerge/v3/src/robo/slack-actions.ts b/Engine/Extras/RoboMerge/v3/src/robo/slack-actions.ts
new file mode 100644
--- /dev/null
+++ b/Engine/Extras/RoboMerge/v3/src/robo/slack-actions.ts
@@ -0,0 +1,510 @@
+// Copyright Epic Games, Inc. All Rights Reserved.
+
+import { ContextualLogger } from '../common/logger';
+import * as request from '../common/request';
+import { OperationReturnType } from './ipc';
+import { Slack, SlackInteractivePayload } from './slack';
+import { BlockageNodeOpParams } from './roboserver';
+
+export interface SlackDispatchContext {
+	sendMessage: (msg: string, args?: any[]) => any
+	botToken: string | undefined
+	slackDomain: string | undefined
+	logger: ContextualLogger
+}
+
+export function extractHostFromSlackPayload(payload: SlackInteractivePayload): string | undefined {
+	try {
+		if (payload.type === 'block_actions' && payload.actions && payload.actions[0]) {
+			return JSON.parse(payload.actions[0].value).host
+		}
+		if (payload.type === 'view_submission' && payload.view) {
+			return JSON.parse(payload.view.private_metadata).host
+		}
+	}
+	catch { /* ignore parse errors */ }
+	return undefined
+}
+
+export async function dispatchSlackInteractivePayload(payload: SlackInteractivePayload, dispatchContext: SlackDispatchContext) {
+	if (payload.type === 'block_actions') {
+		return _handleSlackBlockAction(payload, dispatchContext)
+	}
+	if (payload.type === 'view_submission') {
+		return _handleSlackViewSubmission(payload, dispatchContext)
+	}
+	return { statusCode: 400, message: `Unsupported payload type: ${payload.type}` }
+}
+
+async function _handleSlackBlockAction(payload: SlackInteractivePayload, dispatchContext: SlackDispatchContext) {
+	const action = payload.actions && payload.actions[0]
+	if (!action) {
+		return { statusCode: 400, message: 'No actions in payload' }
+	}
+
+	let params: any
+	try {
+		params = JSON.parse(action.value)
+	}
+	catch {
+		return { statusCode: 400, message: 'Invalid action value JSON' }
+	}
+
+	const { bot, branch, cl, edge, target, targetStream } = params
+	const who = payload.user.name
+
+	if (!bot || !branch || !cl) {
+		return { statusCode: 400, message: 'Missing required params (bot, branch, cl)' }
+	}
+
+	const clStr = String(cl)
+
+	switch (action.action_id) {
+		case 'acknowledge': {
+			const result = edge
+				? await dispatchContext.sendMessage('doEdgeOp', [bot, branch, edge, 'acknowledge', { who, cl: clStr }])
+				: await dispatchContext.sendMessage('doNodeOp', [bot, branch, 'acknowledge', { who, cl: clStr }])
+			return _slackActionResult(result, payload, `Acknowledged by ${who}`, dispatchContext)
+		}
+
+		case 'skip': {
+			if (!edge) {
+				return { statusCode: 400, message: 'skip requires edge param' }
+			}
+			if (!payload.trigger_id) {
+				return { statusCode: 400, message: 'Missing trigger_id for modal' }
+			}
+			return _openSkipReasonModal(payload.trigger_id, params, dispatchContext)
+		}
+
+		case 'stompchanges': {
+			if (!target) {
+				return { statusCode: 400, message: 'stompchanges requires target param' }
+			}
+			if (!payload.trigger_id) {
+				return { statusCode: 400, message: 'Missing trigger_id for modal' }
+			}
+			return _openStompConfirmModal(payload.trigger_id, params, dispatchContext)
+		}
+
+		case 'unlockchanges': {
+			if (!target) {
+				return { statusCode: 400, message: 'unlockchanges requires target param' }
+			}
+			const result = await dispatchContext.sendMessage('doNodeOp', [bot, branch, 'unlockchanges', { who, cl: clStr, target }])
+			return _slackActionResult(result, payload, `Unlock started by ${who}`, dispatchContext)
+		}
+
+		case 'create_shelf': {
+			if (!payload.trigger_id) {
+				return { statusCode: 400, message: 'Missing trigger_id for modal' }
+			}
+			return _openCreateShelfModal(payload.trigger_id, params, who, targetStream, dispatchContext, payload.response_url)
+		}
+
+		default:
+			return { statusCode: 400, message: `Unknown action_id: ${action.action_id}` }
+	}
+}
+
+async function _handleSlackViewSubmission(payload: SlackInteractivePayload, dispatchContext: SlackDispatchContext) {
+	const view = payload.view
+	if (!view) {
+		return { statusCode: 400, message: 'No view in submission payload' }
+	}
+
+	switch (view.callback_id) {
+		case 'skip_reason_modal':   return _handleSkipReasonSubmit(payload, view, dispatchContext)
+		case 'stomp_confirm_modal': return _handleStompConfirmSubmit(payload, view, dispatchContext)
+		default:                    return _handleCreateShelfSubmit(payload, view, dispatchContext)
+	}
+}
+
+async function _handleCreateShelfSubmit(payload: SlackInteractivePayload, view: NonNullable<SlackInteractivePayload['view']>, dispatchContext: SlackDispatchContext) {
+	let params: any
+	try {
+		params = JSON.parse(view.private_metadata)
+	}
+	catch {
+		return { statusCode: 400, message: 'Invalid private_metadata JSON' }
+	}
+
+	const { bot, branch, cl, target, edge } = params
+	const who = payload.user.name
+
+	const values = view.state.values
+	let workspace: string | undefined
+	for (const blockId in values) {
+		for (const actionId in values[blockId]) {
+			const selected = values[blockId][actionId].selected_option
+			if (selected) {
+				workspace = selected.value
+				break
+			}
+		}
+		if (workspace) break
+	}
+
+	if (!workspace) {
+		return {
+			statusCode: 200,
+			message: JSON.stringify({ response_action: 'errors', errors: { workspace_block: 'Please select a workspace' } }),
+			headers: [['Content-Type', 'application/json']]
+		}
+	}
+
+	const acknowledgeSelected = ((values['acknowledge_block'] as any)?.['acknowledge_checkbox']?.selected_options as any[] | undefined)
+		?.some((o: any) => o.value === 'acknowledge') ?? false
+
+	const result = await dispatchContext.sendMessage('doNodeOp', [bot, branch, 'create_shelf', { who, cl: String(cl), target, workspace }])
+	if (result && result.statusCode === 200) {
+		if (acknowledgeSelected) {
+			// Mirror the web UI: acknowledge via edge op when an edge is present
+			if (edge) {
+				await dispatchContext.sendMessage('doEdgeOp', [bot, branch, edge, 'acknowledge', { who, cl: String(cl) }])
+			} else {
+				await dispatchContext.sendMessage('doNodeOp', [bot, branch, 'acknowledge', { who, cl: String(cl) }])
+			}
+		}
+		return { statusCode: 200, message: JSON.stringify({ response_action: 'clear' }), headers: [['Content-Type', 'application/json']] }
+	}
+	return { statusCode: 200, message: JSON.stringify({ response_action: 'errors', errors: { workspace_block: result.message || 'Error creating shelf' } }), headers: [['Content-Type', 'application/json']] }
+}
+
+async function _handleSkipReasonSubmit(payload: SlackInteractivePayload, view: NonNullable<SlackInteractivePayload['view']>, dispatchContext: SlackDispatchContext) {
+	let params: any
+	try {
+		params = JSON.parse(view.private_metadata)
+	}
+	catch {
+		return { statusCode: 400, message: 'Invalid private_metadata JSON' }
+	}
+
+	const { bot, branch, cl, edge } = params
+	const who = payload.user.name
+
+	const values = view.state.values
+	let selectedReason: string | undefined
+	for (const blockId in values) {
+		for (const actionId in values[blockId]) {
+			const sel = values[blockId][actionId].selected_option
+			if (sel) { selectedReason = sel.value; break }
+		}
+		if (selectedReason) break
+	}
+
+	if (!selectedReason) {
+		return {
+			statusCode: 200,
+			message: JSON.stringify({ response_action: 'errors', errors: { reason_block: 'Please select a reason' } }),
+			headers: [['Content-Type', 'application/json']]
+		}
+	}
+
+	const reasonText = selectedReason === 'notrelevant'
+		? 'Work is not relevant to other branches'
+		: 'User will redo work in the merge target'
+
+	const result = await dispatchContext.sendMessage('doEdgeOp', [bot, branch, edge, 'set_last_cl', {
+		who, cl: String(cl), reason: reasonText, unblock: 'true', isSkip: 'true'
+	}])
+
+	if (result && result.statusCode === 200) {
+		if (payload.response_url && dispatchContext.botToken) {
+			request.post({
+				url: payload.response_url,
+				body: JSON.stringify({ replace_original: false, text: `✓ Skipped by ${who} (${reasonText})` }),
+				contentType: 'application/json'
+			}).catch(() => { /* best-effort */ })
+		}
+		return { statusCode: 200, message: JSON.stringify({ response_action: 'clear' }), headers: [['Content-Type', 'application/json']] }
+	}
+	return { statusCode: 200, message: JSON.stringify({ response_action: 'errors', errors: { reason_block: result?.message || 'Error performing skip' } }), headers: [['Content-Type', 'application/json']] }
+}
+
+async function _handleStompConfirmSubmit(payload: SlackInteractivePayload, view: NonNullable<SlackInteractivePayload['view']>, dispatchContext: SlackDispatchContext) {
+	let params: any
+	try {
+		params = JSON.parse(view.private_metadata)
+	}
+	catch {
+		return { statusCode: 400, message: 'Invalid private_metadata JSON' }
+	}
+
+	const { bot, branch, cl, target } = params
+	const who = payload.user.name
+
+	const result = await dispatchContext.sendMessage('doNodeOp', [bot, branch, 'stompchanges', { who, cl: String(cl), target }])
+
+	if (result && result.statusCode === 200) {
+		return { statusCode: 200, message: JSON.stringify({ response_action: 'clear' }), headers: [['Content-Type', 'application/json']] }
+	}
+
+	// Stomp failed — update the modal to show the error so the user sees it
+	const errorText = result?.message || 'Unknown error performing stomp'
+	const errorModal = {
+		type: 'modal',
+		callback_id: 'stomp_confirm_modal',
+		title: { type: 'plain_text', text: 'Stomp Failed' },
+		close: { type: 'plain_text', text: 'Close' },
+		private_metadata: view.private_metadata,
+		blocks: [
+			{
+				type: 'section',
+				text: { type: 'mrkdwn', text: `:x: *Stomp failed*\n\n${errorText}` }
+			}
+		]
+	}
+	return {
+		statusCode: 200,
+		message: JSON.stringify({ response_action: 'update', view: errorModal }),
+		headers: [['Content-Type', 'application/json']]
+	}
+}
+
+async function _openSkipReasonModal(triggerId: string, params: BlockageNodeOpParams, dispatchContext: SlackDispatchContext) {
+	if (!dispatchContext.botToken || !dispatchContext.slackDomain) {
+		return { statusCode: 500, message: 'Slack not configured for modal' }
+	}
+	const modal = {
+		type: 'modal',
+		callback_id: 'skip_reason_modal',
+		title: { type: 'plain_text', text: 'Skip Merge' },
+		submit: { type: 'plain_text', text: 'Skip' },
+		close: { type: 'plain_text', text: 'Cancel' },
+		private_metadata: JSON.stringify(params),
+		blocks: [
+			{
+				type: 'section',
+				text: { type: 'mrkdwn', text: `*Skip merging CL ${params.cl} to ${params.edge}?*\nThis will skip the merge and allow Robomerge to continue.` }
+			},
+			{
+				type: 'input',
+				block_id: 'reason_block',
+				label: { type: 'plain_text', text: 'Reason for skipping' },
+				element: {
+					type: 'static_select',
+					action_id: 'reason_select',
+					placeholder: { type: 'plain_text', text: 'Select a reason' },
+					options: [
+						{ text: { type: 'plain_text', text: 'Work is not relevant to the target branch' }, value: 'notrelevant' },
+						{ text: { type: 'plain_text', text: 'I will manually redo work in the merge targets' }, value: 'willredo' }
+					]
+				}
+			}
+		]
+	}
+	const slack = new Slack({ id: '', botToken: dispatchContext.botToken, userToken: '' }, dispatchContext.slackDomain, dispatchContext.logger)
+	const result = await slack.openModal(triggerId, modal)
+	return result?.ok ? { statusCode: 200, message: '' } : { statusCode: 500, message: `Failed to open modal: ${result?.error}` }
+}
+
+async function _openStompConfirmModal(triggerId: string, params: BlockageNodeOpParams, dispatchContext: SlackDispatchContext) {
+	if (!dispatchContext.botToken || !dispatchContext.slackDomain) {
+		return { statusCode: 500, message: 'Slack not configured for modal' }
+	}
+	const modal = {
+		type: 'modal',
+		callback_id: 'stomp_confirm_modal',
+		title: { type: 'plain_text', text: 'Stomp Changes' },
+		submit: { type: 'plain_text', text: 'Confirm Stomp' },
+		close: { type: 'plain_text', text: 'Cancel' },
+		private_metadata: JSON.stringify(params),
+		blocks: [
+			{
+				type: 'section',
+				text: {
+					type: 'mrkdwn',
+					text: `:warning: *Stomp changes in ${params.target} with CL ${params.cl}?*\n\nRobomerge will verify the stomp before proceeding. Non-binary file conflicts will block the stomp.`
+				}
+			}
+		]
+	}
+	const slack = new Slack({ id: '', botToken: dispatchContext.botToken, userToken: '' }, dispatchContext.slackDomain, dispatchContext.logger)
+	const result = await slack.openModal(triggerId, modal)
+	return result?.ok ? { statusCode: 200, message: '' } : { statusCode: 500, message: `Failed to open modal: ${result?.error}` }
+}
+
+async function _openCreateShelfModal(triggerId: string, params: any, who: string, targetStream: string | undefined, dispatchContext: SlackDispatchContext, responseUrl?: string) {
+	if (!dispatchContext.botToken || !dispatchContext.slackDomain) {
+		return { statusCode: 500, message: 'Slack not configured for modal' }
+	}
+
+	const { cl, target } = params
+
+	dispatchContext.logger.info(`_openCreateShelfModal: user='${who}', CL=${cl}, target='${target}', targetStream='${targetStream}'`)
+
+	const workspacesResult = await dispatchContext.sendMessage('getWorkspaces', [who, ''])
+	const allWorkspaces: Array<{client: string, Stream?: string}> = workspacesResult?.data ?? []
+	dispatchContext.logger.info(`_openCreateShelfModal: getWorkspaces returned ${allWorkspaces.length} total workspaces for user='${who}'`)
+
+	// Slack option text must be ≤ 75 chars; option value must be ≤ 150 chars.
+	const truncateLabel = (label: string) =>
+		label.length <= 75 ? label : label.substring(0, 72) + '...'
+
+	const makeClientOption = (ws: {client: string}) => ({
+		text: { type: 'plain_text', text: truncateLabel(ws.client) },
+		value: ws.client.length > 150 ? ws.client.substring(0, 150) : ws.client
+	})
+
+	// For inexact matches, always show "client (stream-local)" where stream-local is the stream
+	// name with the shared depot prefix stripped — the depot is shown in the group label instead.
+	const makeInexactOption = (ws: {client: string, Stream?: string}, depotPath: string) => {
+		const streamLocal = ws.Stream ? ws.Stream.replace(depotPath, '') : ws.Stream ?? ''
+		const label = streamLocal ? `${ws.client} (${streamLocal})` : ws.client
+		return {
+			text: { type: 'plain_text', text: truncateLabel(label) },
+			value: ws.client.length > 150 ? ws.client.substring(0, 150) : ws.client
+		}
+	}
+
+	// Sort alphabetically within each category (mirrors the web UI)
+	const sorted = [...allWorkspaces].sort((a, b) => a.client.localeCompare(b.client))
+
+	// Categorise workspaces the same way the web UI does:
+	//   exact   — Stream matches targetStream exactly  (shown first, bold on web)
+	//   inexact — Stream shares the same depot path    (shown second)
+	//   plain   — no Stream field                      (shown when no targetStream)
+	let exactOptions: ReturnType<typeof makeClientOption>[] = []
+	let inexactOptions: ReturnType<typeof makeInexactOption>[] = []
+	let plainOptions: ReturnType<typeof makeClientOption>[] = []
+	let depotPath = ''
+
+	if (targetStream) {
+		const depotPathMatch = targetStream.match(/(\/\/[\w\-]+\/)/)
+		depotPath = depotPathMatch ? depotPathMatch[0] : ''
+		for (const ws of sorted) {
+			if (!ws.Stream) continue
+			if (ws.Stream === targetStream) {
+				// Exact match — stream shown in group label, omit it from individual options
+				exactOptions.push(makeClientOption(ws))
+			} else if (depotPath && ws.Stream.startsWith(depotPath)) {
+				inexactOptions.push(makeInexactOption(ws, depotPath))
+			}
+		}
+	} else {
+		for (const ws of sorted) {
+			if (!ws.Stream) {
+				plainOptions.push(makeClientOption(ws))
+			}
+		}
+	}
+
+	// Build the select element — use option_groups when there are multiple categories
+	// so the grouping is visible in the dropdown (mirrors web UI radio button sections).
+	const totalOptions = exactOptions.length + inexactOptions.length + plainOptions.length
+	if (totalOptions === 0) {
+		const msg = allWorkspaces.length === 0
+			? `No workspaces found for user '${who}'. Please use the web UI to create a shelf.`
+			: `No workspaces matching stream '${targetStream}' found for user '${who}'. Please use the web UI to create a shelf.`
+		dispatchContext.logger.warn(`_openCreateShelfModal: ${msg}`)
+		if (responseUrl) {
+			request.post({ url: responseUrl, body: JSON.stringify({ replace_original: false, response_type: 'ephemeral', text: `:warning: ${msg}` }), contentType: 'application/json' }).catch(() => { /* best-effort */ })
+		}
+		return { statusCode: 200, message: msg }
+	}
+
+	// Group label for exact matches — include the stream name since individual options omit it.
+	// Slack group labels are also capped at 75 chars.
+	const exactGroupLabel = truncateLabel(targetStream ? `Matching workspaces (${targetStream})` : 'Matching workspaces')
+
+	let selectElement: any
+	if (exactOptions.length > 0 && inexactOptions.length > 0) {
+		// Two distinct groups — use option_groups so the headers are visible
+		selectElement = {
+			type: 'static_select',
+			action_id: 'workspace_select',
+			placeholder: { type: 'plain_text', text: 'Select a workspace' },
+			option_groups: [
+				{ label: { type: 'plain_text', text: exactGroupLabel }, options: exactOptions.slice(0, 50) },
+				{ label: { type: 'plain_text', text: truncateLabel(`Other workspaces in this depot (${depotPath})`) }, options: inexactOptions.slice(0, 50) }
+			]
+		}
+	} else if (exactOptions.length > 0) {
+		// Only exact matches — single group with stream name in header, options are client-name only
+		selectElement = {
+			type: 'static_select',
+			action_id: 'workspace_select',
+			placeholder: { type: 'plain_text', text: 'Select a workspace' },
+			option_groups: [
+				{ label: { type: 'plain_text', text: exactGroupLabel }, options: exactOptions.slice(0, 100) }
+			]
+		}
+	} else {
+		// Only inexact or plain workspaces — flat options list
+		const flat = [...inexactOptions, ...plainOptions].slice(0, 100)
+		selectElement = {
+			type: 'static_select',
+			action_id: 'workspace_select',
+			placeholder: { type: 'plain_text', text: 'Select a workspace' },
+			options: flat
+		}
+	}
+
+	const modal = {
+		type: 'modal',
+		callback_id: 'create_shelf_modal',
+		title: { type: 'plain_text', text: 'Create Shelf' },
+		submit: { type: 'plain_text', text: 'Create Shelf' },
+		close: { type: 'plain_text', text: 'Cancel' },
+		private_metadata: JSON.stringify(params),
+		blocks: [{
+			type: 'section',
+			text: { type: 'mrkdwn', text: `Select workspace to create shelf for CL ${cl} → *${target}*` }
+		}, {
+			type: 'input',
+			block_id: 'workspace_block',
+			label: { type: 'plain_text', text: 'Workspace' },
+			element: selectElement
+		}, {
+			type: 'input',
+			block_id: 'acknowledge_block',
+			optional: true,
+			label: { type: 'plain_text', text: ' ' },
+			element: {
+				type: 'checkboxes',
+				action_id: 'acknowledge_checkbox',
+				initial_options: [{
+					text: { type: 'plain_text', text: 'Acknowledge conflict on successful shelf creation' },
+					value: 'acknowledge'
+				}],
+				options: [{
+					text: { type: 'plain_text', text: 'Acknowledge conflict on successful shelf creation' },
+					value: 'acknowledge'
+				}]
+			}
+		}]
+	}
+
+	const slack = new Slack(
+		{ id: '', botToken: dispatchContext.botToken, userToken: '' },
+		dispatchContext.slackDomain,
+		dispatchContext.logger
+	)
+	const modalResult = await slack.openModal(triggerId, modal)
+	if (modalResult && modalResult.ok) {
+		dispatchContext.logger.info(`_openCreateShelfModal: modal opened successfully for user='${who}', CL=${cl}`)
+		return { statusCode: 200, message: '' }
+	}
+	const errDetail = modalResult?.error ?? 'unknown error'
+	const errMessages = modalResult?.response_metadata?.messages
+	dispatchContext.logger.error(`_openCreateShelfModal: views.open failed for user='${who}', CL=${cl}: ${errDetail}${errMessages ? ' — ' + JSON.stringify(errMessages) : ''}`)
+	if (responseUrl) {
+		request.post({ url: responseUrl, body: JSON.stringify({ replace_original: false, response_type: 'ephemeral', text: `:x: Failed to open Create Shelf dialog: ${errDetail}` }), contentType: 'application/json' }).catch(() => { /* best-effort */ })
+	}
+	return { statusCode: 500, message: `Failed to open modal: ${errDetail}` }
+}
+
+async function _slackActionResult(result: OperationReturnType, payload: SlackInteractivePayload, successText: string, dispatchContext: SlackDispatchContext) {
+	if (result && result.statusCode === 200) {
+		if (payload.response_url && dispatchContext.botToken) {
+			request.post({
+				url: payload.response_url,
+				body: JSON.stringify({ replace_original: false, text: successText }),
+				contentType: 'application/json'
+			}).catch(() => { /* best-effort */ })
+		}
+		return { statusCode: 200, message: '' }
+	}
+	return { statusCode: 200, message: JSON.stringify({ text: `Error: ${result && result.message}` }), headers: [['Content-Type', 'application/json']] }
+}
diff --git a/Engine/Extras/RoboMerge/v3/src/robo/slack-socket.ts b/Engine/Extras/RoboMerge/v3/src/robo/slack-socket.ts
new file mode 100644
--- /dev/null
+++ b/Engine/Extras/RoboMerge/v3/src/robo/slack-socket.ts
@@ -0,0 +1,182 @@
+// Copyright Epic Games, Inc. All Rights Reserved.
+
+import { SlackInteractivePayload } from './slack'
+import { dispatchSlackInteractivePayload, extractHostFromSlackPayload, SlackDispatchContext } from './slack-actions'
+
+const RECONNECT_DELAY_MS = 5000
+
+export class SlackSocket {
+	private ws: WebSocket | null = null
+	private reconnectTimer: ReturnType<typeof setTimeout> | null = null
+	private stopped = false
+
+	constructor(
+		private appToken: string,
+		private externalUrl: string,
+		private ctx: SlackDispatchContext
+	) {}
+
+	async start() {
+		this.stopped = false
+		try {
+			const wsUrl = await this.getWebSocketUrl()
+			this.connect(wsUrl)
+		}
+		catch (err) {
+			this.ctx.logger.error(`Slack Socket Mode: failed to get WebSocket URL: ${err}`)
+			this.scheduleReconnect()
+		}
+	}
+
+	stop() {
+		this.stopped = true
+		if (this.reconnectTimer) {
+			clearTimeout(this.reconnectTimer)
+			this.reconnectTimer = null
+		}
+		this.ws?.close()
+		this.ws = null
+	}
+
+	private async getWebSocketUrl(): Promise<string> {
+		const resp = await fetch('https://slack.com/api/apps.connections.open', {
+			method: 'POST',
+			headers: {
+				'Authorization': `Bearer ${this.appToken}`,
+				'Content-Type': 'application/x-www-form-urlencoded'
+			}
+		})
+		const json = await resp.json() as any
+		if (!json.ok) {
+			throw new Error(`apps.connections.open failed: ${json.error}`)
+		}
+		return json.url
+	}
+
+	private connect(url: string) {
+		this.ctx.logger.info('Slack Socket Mode: connecting...')
+		const ws = new WebSocket(url)
+		this.ws = ws
+
+		ws.addEventListener('open', () => {
+			this.ctx.logger.info('Slack Socket Mode: connected')
+		})
+
+		ws.addEventListener('message', (event: MessageEvent) => {
+			this.handleMessage(event.data as string)
+		})
+
+		ws.addEventListener('close', (event: Event) => {
+			if (this.ws !== ws) return  // already replaced
+			this.ctx.logger.warn(`Slack Socket Mode: closed (code=${(event as any).code}), reconnecting...`)
+			this.ws = null
+			this.scheduleReconnect()
+		})
+
+		ws.addEventListener('error', () => {
+			this.ctx.logger.error('Slack Socket Mode: WebSocket error')
+		})
+	}
+
+	private handleMessage(data: string) {
+		let msg: any
+		try {
+			msg = JSON.parse(data)
+		}
+		catch {
+			return
+		}
+
+		if (msg.type === 'hello') {
+			this.ctx.logger.verbose('Slack Socket Mode: hello received')
+			return
+		}
+
+		if (msg.type === 'disconnect') {
+			this.ctx.logger.info('Slack Socket Mode: server requested disconnect, reconnecting...')
+			this.ws?.close()
+			return
+		}
+
+		if (msg.type === 'interactive') {
+			const payload = msg.payload as SlackInteractivePayload
+			if (payload?.type === 'view_submission') {
+				// For view_submission, the response_action must be included in the envelope ACK.
+				// Do NOT ACK immediately — wait for the dispatch result (within Slack's 3-second window).
+				this.handleViewSubmission(msg.envelope_id, payload)
+				return
+			}
+			// For block_actions, acknowledge immediately and handle asynchronously.
+			if (msg.envelope_id) {
+				this.ws?.send(JSON.stringify({ envelope_id: msg.envelope_id }))
+			}
+			this.handleInteractive(payload)
+			return
+		}
+
+		// Acknowledge all other envelopes immediately
+		if (msg.envelope_id) {
+			this.ws?.send(JSON.stringify({ envelope_id: msg.envelope_id }))
+		}
+	}
+
+	private handleInteractive(payload: SlackInteractivePayload) {
+		if (!payload) return
+
+		// With Socket Mode all connected instances receive every event.
+		// Only handle actions targeted at this instance.
+		const host = extractHostFromSlackPayload(payload)
+		if (host && host !== this.externalUrl) {
+			return
+		}
+
+		dispatchSlackInteractivePayload(payload, this.ctx).catch(err => {
+			this.ctx.logger.error(`Slack Socket Mode: dispatch error: ${err}`)
+		})
+	}
+
+	private async handleViewSubmission(envelopeId: string, payload: SlackInteractivePayload) {
+		if (!payload) {
+			if (envelopeId) this.ws?.send(JSON.stringify({ envelope_id: envelopeId }))
+			return
+		}
+
+		// With Socket Mode all connected instances receive every event.
+		// Only handle submissions targeted at this instance.
+		const host = extractHostFromSlackPayload(payload)
+		if (host && host !== this.externalUrl) {
+			if (envelopeId) this.ws?.send(JSON.stringify({ envelope_id: envelopeId }))
+			return
+		}
+
+		let ackPayload: any = {}
+		try {
+			const result = await dispatchSlackInteractivePayload(payload, this.ctx)
+			// Extract the response_action JSON from the result message and include it in the ACK
+			if (result?.message && typeof result.message === 'string') {
+				try {
+					const parsed = JSON.parse(result.message)
+					if (parsed?.response_action) {
+						ackPayload.payload = parsed
+					}
+				}
+				catch { /* message isn't JSON — no response_action to send */ }
+			}
+		}
+		catch (err) {
+			this.ctx.logger.error(`Slack Socket Mode: view submission dispatch error: ${err}`)
+		}
+
+		if (envelopeId) {
+			this.ws?.send(JSON.stringify({ envelope_id: envelopeId, ...ackPayload }))
+		}
+	}
+
+	private scheduleReconnect() {
+		if (this.stopped || this.reconnectTimer) return
+		this.reconnectTimer = setTimeout(() => {
+			this.reconnectTimer = null
+			this.start()
+		}, RECONNECT_DELAY_MS)
+	}
+}
diff --git a/Engine/Extras/RoboMerge/v3/src/robo/slack.ts b/Engine/Extras/RoboMerge/v3/src/robo/slack.ts
--- a/Engine/Extras/RoboMerge/v3/src/robo/slack.ts
+++ b/Engine/Extras/RoboMerge/v3/src/robo/slack.ts
@@ -25,6 +25,45 @@ export interface SlackLinkButtonAction {
 	value?: string
 }
 
+// https://api.slack.com/reference/block-kit/block-elements#button
+export interface SlackInteractiveButton {
+	type: 'button'
+	text: { type: 'plain_text'; text: string }
+	action_id: string       // operation name, e.g. 'acknowledge'
+	value: string           // JSON-encoded params: {bot, branch, cl, edge?, target?}
+	style?: 'primary' | 'danger'
+}
+
+// https://api.slack.com/reference/block-kit/blocks#actions
+export interface SlackActionsBlock {
+	type: 'actions'
+	elements: SlackInteractiveButton[]
+}
+
+// Inbound payload from Slack when a user clicks an interactive button
+export interface SlackInteractiveAction {
+	action_id: string
+	value: string
+}
+
+export interface SlackInteractivePayload {
+	type: 'block_actions' | 'view_submission'
+	user: { id: string; name: string }
+	channel?: { id: string }
+	message?: { ts: string }
+	response_url?: string
+	// block_actions
+	actions?: SlackInteractiveAction[]
+	// view_submission
+	view?: {
+		callback_id: string
+		private_metadata: string
+		state: { values: { [blockId: string]: { [actionId: string]: { selected_option?: { value: string } } } } }
+	}
+	// triggers view.open for modal
+	trigger_id?: string
+}
+
 // https://api.slack.com/docs/interactive-message-field-guide#attachment_fields
 export interface SlackAttachment {
 	text?: string
@@ -60,6 +99,8 @@ export interface SlackMessage {
 
 	// An array because the Slack API expects it to be
 	attachments?: SlackAttachment[]
+	// Block Kit blocks (interactive buttons etc.)
+	blocks?: SlackActionsBlock[]
 	// Direct Message support
 	channel: string
 	// Allows Markdown formatting in messages
@@ -149,6 +190,10 @@ export class Slack {
 		return this.post('chat.update', args)
 	}
 
+	delete(ts: string, channel: string) {
+		return this.post('chat.delete', { ts, channel }, true)
+	}
+
 	uploadFile(file: SlackFile) {
 		let getURLargs: any = {
 			filename: file.filename,
@@ -276,6 +321,14 @@ export class Slack {
 //	should now overload so that string only puts text in args
 
 // look out for makeArgs with just string - support if necessary
+	async getUserInfo(userId: string) {
+		return this.get('users.info', {user: userId}, true)
+	}
+
+	async openModal(triggerId: string, view: any) {
+		return this.post('views.open', {trigger_id: triggerId, view: JSON.stringify(view)}, true)
+	}
+
 	private makeArgs(message: SlackMessage) {
 		const args: {[arg: string]: any} = {}
 
@@ -289,12 +342,19 @@ export class Slack {
 			args.attachments = args.attachments.concat(message.attachments)
 		}
 
+		// Block Kit blocks are passed at the top level, not inside attachments.
+		// Pass explicitly even when empty so chat.update clears existing blocks.
+		if (message.blocks !== undefined) {
+			args.blocks = message.blocks
+		}
+
+		const SKIP_ATTACHMENT_KEYS = new Set(['style', 'attachments', 'blocks'])
 		const opts = message as unknown as {[key: string]: any}
 		for (const key in opts) {
 			if (MAIN_MESSAGE_FIELDS.has(key)) {
 				args[key] = opts[key]
 			}
-			else if (key !== 'style') {
+			else if (!SKIP_ATTACHMENT_KEYS.has(key)) {
 				attch[key] = opts[key]
 			}
 		}
diff --git a/Engine/Extras/RoboMerge/v3/src/robo/vault.ts b/Engine/Extras/RoboMerge/v3/src/robo/vault.ts
--- a/Engine/Extras/RoboMerge/v3/src/robo/vault.ts
+++ b/Engine/Extras/RoboMerge/v3/src/robo/vault.ts
@@ -38,6 +38,10 @@ export class Vault {
         return this._vault && this._vault['slack-tokens']
     }
 
+    get slackAppToken(): string | undefined {
+        return this._vault && this._vault['slack-tokens'] && this._vault['slack-tokens']['app']
+    }
+
     get cookieKey() {
         return this._vault && this._vault['cookie-key']
     }
