Don't pick workspace for user that is on an edge that does not provide its address in the server record

diff --git a/Engine/Extras/RoboMerge/v3/src/common/p4util.ts b/Engine/Extras/RoboMerge/v3/src/common/p4util.ts
--- a/Engine/Extras/RoboMerge/v3/src/common/p4util.ts
+++ b/Engine/Extras/RoboMerge/v3/src/common/p4util.ts
@@ -95,8 +95,7 @@ export async function chooseBestWorkspaceForUser(p4: PerforceContext, user: stri
 
 	if (workspaces.length > 0) {
 		workspaces.sort((a,b) => b.Access - a.Access)
-		// default to the first workspace
-		let targetWorkspace = workspaces[0]
+		let targetWorkspace: ClientSpec | undefined
 
 		// if this is a stream branch, do some better match-up
 		if (stream) {
@@ -104,19 +103,30 @@ export async function chooseBestWorkspaceForUser(p4: PerforceContext, user: stri
 			// find the stream with the closest match
 			let target_match = 0
 			for (let def of workspaces) {
-				let stream = def.Stream
-				if (stream) {
-					const matchlen = matchPrefix(stream.toLowerCase(), branch_stream)
-					if (matchlen == branch_stream.length) {
+				const matchlen = def.Stream ? matchPrefix(def.Stream.toLowerCase(), branch_stream) : 0
+				if (matchlen == branch_stream.length) {
+					if (await p4.getWorkspaceEdgeServer(def.client) !== false) {
 						return def
 					}
-					else if (matchlen > target_match) {
-						target_match = matchlen
+				}
+				else if (matchlen > target_match || !targetWorkspace) {
+					// better prefix match, or no reachable default found yet
+					if (await p4.getWorkspaceEdgeServer(def.client) !== false) {
+						if (matchlen > target_match) {
+							target_match = matchlen
+						}
 						targetWorkspace = def
 					}
 				}
 			}
 		}
+		else {
+			for (const def of workspaces) {
+				if (await p4.getWorkspaceEdgeServer(def.client) !== false) {
+					return def
+				}
+			}
+		}
 		return targetWorkspace
 	}
 
diff --git a/Engine/Extras/RoboMerge/v3/src/common/perforce.ts b/Engine/Extras/RoboMerge/v3/src/common/perforce.ts
--- a/Engine/Extras/RoboMerge/v3/src/common/perforce.ts
+++ b/Engine/Extras/RoboMerge/v3/src/common/perforce.ts
@@ -879,13 +879,14 @@ export class PerforceContext {
 		return result
 	}
 
-	async getWorkspaceEdgeServer(workspaceName: string): Promise<EdgeServer | null> {
-		const serverIdLine = await this._execP4(null, ['-ztag', '-F', '%ServerID%', 'client', '-o', workspaceName]) 
+	async getWorkspaceEdgeServer(workspaceName: string): Promise<EdgeServer | null | false> {
+		const serverIdLine = await this._execP4(null, ['-ztag', '-F', '%ServerID%', 'client', '-o', workspaceName])
 		if (serverIdLine) {
 			const serverId = serverIdLine.trim()
 			const address = await this.getEdgeServerAddress(serverId)
 			if (!address) {
-				throw new Error(`Couldn't find address for edge server '${serverId}'`)
+				this.logger.warn(`No address record for edge server '${serverId}' (workspace '${workspaceName}') — cannot use this workspace`)
+				return false
 			}
 			return {id: serverId, address: address.trim()}
 		}
@@ -1671,7 +1672,7 @@ export class PerforceContext {
 		const edgeServer = await this.getWorkspaceEdgeServer(client)
 		const args = ['revert', ...(opts||[]), '-C', client, ...files]
 		try {
-			const results = await this.execAndParse(null, args, {serverAddress: edgeServer?.address})
+			const results = await this.execAndParse(null, args, {serverAddress: edgeServer ? edgeServer.address : undefined})
 
 			// Files can fail to be reverted because of moves, lets try and resolve that
 			const movedFiles = 
