Fix track change failing to back track in certain cases
Fix a backout of a move resulting in duplicate entries from the same stream in the track change graph

diff --git a/Engine/Extras/RoboMerge/v3/src/robo/trackchange.ts b/Engine/Extras/RoboMerge/v3/src/robo/trackchange.ts
--- a/Engine/Extras/RoboMerge/v3/src/robo/trackchange.ts
+++ b/Engine/Extras/RoboMerge/v3/src/robo/trackchange.ts
@@ -229,14 +229,15 @@ export async function trackChange(robo: RoboMerge, logger: ContextualLogger, ini
                     sourceChain.push(parentCL)
                 }
                 else {
-                    for (const entry of desc.entries) {
+                    for (let i = 0; i < desc.entries.length; i++) {
+                        const entry = desc.entries[i]
                         // integrated for move/delete will point at the paired move/add not where it was merged to in another stream, so not useful to evaluate it
                         if (entry.action != 'move/delete') {
                             // Do an integrated check
                             const integrated = (await p4.integrated(null, entry.depotFile))
-                                                        .filter(   x => x.how.includes('from') 
-                                                                && x.toFile == entry.depotFile 
-                                                                && x.endToRev == `#${entry.rev}` 
+                                                        .filter(   x => x.how.includes('from')
+                                                                && x.toFile == entry.depotFile
+                                                                && x.endToRev == `#${entry.rev}`
                                                                 && getStreamFromPath(x.toFile) != getStreamFromPath(x.fromFile))
                             if (integrated.length > 0) {
                                 const fstat = await p4.fstat(null, `${integrated[0].fromFile}${integrated[0].endFromRev}`)
@@ -384,6 +385,9 @@ export async function trackChange(robo: RoboMerge, logger: ContextualLogger, ini
                     const integrated = await p4.integrated(null, entry.depotFile, {intoOnly: true, startCL: clToConsider.cl})
                     if (integrated.length > 0) {
                         for (let integ of integrated) {
+                            if (getStreamFromPath(integ.fromFile) === getStreamFromPath(integ.toFile)) {
+                                continue  // same-stream integration — backout/revert, not a forward merge
+                            }
                             const startToRev = integ.startToRev == "#none" ? 0 : parseInt(integ.startToRev.slice(1))
                             const endToRev = parseInt(integ.endToRev.slice(1))
                             if (entry.rev <= endToRev && entry.rev > startToRev)
