Fix changes failing when neither since or until are supplied

diff --git a/Engine/Extras/RoboMerge/v3/src/common/perforce.ts b/Engine/Extras/RoboMerge/v3/src/common/perforce.ts
--- a/Engine/Extras/RoboMerge/v3/src/common/perforce.ts
+++ b/Engine/Extras/RoboMerge/v3/src/common/perforce.ts
@@ -676,7 +676,8 @@ export class PerforceContext {
 	 */
 	changes(path_in: string, opts?: ChangesOpts): Promise<Change[]> {
 		opts = opts || {}
-		const range = '@' + (opts.since ? (opts.until ? `${opts.since},${opts.until}` : `${opts.since},now`) : (opts.until ? `${opts.until}` : ''))
+		const rangeStr = opts.since ? (opts.until ? `${opts.since},${opts.until}` : `${opts.since},now`) : (opts.until ? `${opts.until}` : '')
+		const range = rangeStr ? '@' + rangeStr : ''
 		const path = path_in + range
 		const args = ['changes', '-l',
 			(opts.status ? `-s${opts.status}` : '-ssubmitted'),
