Make robomerge look harder for related changelists that came from a merge conflict and then were submitted as an edit instead of a merge w/ edit
Fixed potential issue tracking changes out of an edge with multiple external servers

diff --git a/Engine/Extras/RoboMerge/v3/src/robo/trackchange.ts b/Engine/Extras/RoboMerge/v3/src/robo/trackchange.ts
--- a/Engine/Extras/RoboMerge/v3/src/robo/trackchange.ts
+++ b/Engine/Extras/RoboMerge/v3/src/robo/trackchange.ts
@@ -308,6 +308,7 @@ export async function trackChange(robo: RoboMerge, logger: ContextualLogger, ini
         }
 
         const externalServerTargets: Node[] = []
+        let sameServerTargets: Node[] = []
         let includeInResults = false
         let hasAutomergeTarget = false
         let streamDisplayName = ""
@@ -331,7 +332,10 @@ export async function trackChange(robo: RoboMerge, logger: ContextualLogger, ini
                         hasAutomergeTarget = edge.source.stream == changeToConsider.node.stream
                     }
                 }
-                if (!PerforceContext.sameServerID(clToConsider.serverId, edge.target.serverID)) {
+                if (PerforceContext.sameServerID(clToConsider.serverId, edge.target.serverID)) {
+                    sameServerTargets.push(edge.target)
+                }
+                else {
                     externalServerTargets.push(edge.target)
                 }
             }
@@ -410,7 +414,7 @@ export async function trackChange(robo: RoboMerge, logger: ContextualLogger, ini
                     }
                 }
 
-                if (!runIntegrated && externalServerTargets.length > 0) {
+                if (!runIntegrated && externalServerTargets.length == 0) {
                     break
                 }
             }
@@ -422,6 +426,45 @@ export async function trackChange(robo: RoboMerge, logger: ContextualLogger, ini
                 await getMoreDescEntries(p4, clToConsider, changeToConsider.desc)
             }
         }
+
+        // Fallback for same-server targets not found via p4 integrated (e.g. conflict resolutions
+        // where the user converted integrated files to edits, leaving no p4 integration records).
+        sameServerTargets = sameServerTargets.filter(target =>
+            !changeToConsider.destCLs.some(cl => { const n = changes.get(cl)?.node; return n && typeof n !== 'string' && n.stream === target.stream })
+        )
+        if (sameServerTargets.length > 0) {
+            for (const entry of changeToConsider.desc.entries) {
+                if (entry.action != 'move/delete') {
+                    const srcStream = getStreamFromPath(entry.depotFile)
+                    if (srcStream) {
+                        for (let i=sameServerTargets.length-1; i>=0; i--) {
+                            const targetPath = entry.depotFile.replace(srcStream, sameServerTargets[i].stream)
+                            const targetCLs = await p4.changes(targetPath, {since: clToConsider.cl})
+                            for (const targetCL of targetCLs) {
+                                if (targetCL.desc.includes(clToConsider.cl.toString())) {
+                                    const newCL = {cl: targetCL.change, serverId: clToConsider.serverId}
+                                    if (!changeToConsider.destCLs.some(cl => compareTrackedCL(cl, newCL))) {
+                                        changeToConsider.destCLs.push(newCL)
+                                        const existingChange = changes.get(newCL)
+                                        if (existingChange) {
+                                            existingChange.sourceCL = clToConsider
+                                        }
+                                        else {
+                                            await gatherCLInfo(newCL, {parentCL: clToConsider})
+                                        }
+                                    }
+                                    sameServerTargets.splice(i, 1)
+                                    break
+                                }
+                            }
+                        }
+                    }
+                }
+                if (sameServerTargets.length == 0) {
+                    break
+                }
+            }
+        }
         for (const destCL of changeToConsider.destCLs) {
             if (!clsToConsider.find(cl => compareTrackedCL(cl, destCL))) {
                 clsToConsider.push(destCL)
