Fix race condition in functional test

diff --git a/Engine/Extras/RoboMerge/v3/functional_tests/framework/src/tests/exclusive-checkout.ts b/Engine/Extras/RoboMerge/v3/functional_tests/framework/src/tests/exclusive-checkout.ts
--- a/Engine/Extras/RoboMerge/v3/functional_tests/framework/src/tests/exclusive-checkout.ts
+++ b/Engine/Extras/RoboMerge/v3/functional_tests/framework/src/tests/exclusive-checkout.ts
@@ -26,10 +26,10 @@ export class ExclusiveCheckout extends FunctionalTest {
 		])
 	}
 
-	run() {
+	async run() {
 		// second user commits change in Release while first user has file in main checked out
-		P4Util.editFile(this.releaseUser2Client, 'test.uasset', 'New content')
-		P4Util.editFile(this.releaseUser2Client, 'test2.uasset', 'New content')
+		await P4Util.editFile(this.releaseUser2Client, 'test.uasset', 'New content')
+		await P4Util.editFile(this.releaseUser2Client, 'test2.uasset', 'New content')
 		return P4Util.submit(this.releaseUser2Client, "Submit by second user")
 	}
 
