Non default bots will defer processing of changelists until a default bot has processed it to avoid syntax errors that should have impacted the non-default bot but were ignored because they were not targeted to that bot.

diff --git a/Engine/Extras/RoboMerge/v3/functional_tests/framework/src/BranchState.ts b/Engine/Extras/RoboMerge/v3/functional_tests/framework/src/BranchState.ts
--- a/Engine/Extras/RoboMerge/v3/functional_tests/framework/src/BranchState.ts
+++ b/Engine/Extras/RoboMerge/v3/functional_tests/framework/src/BranchState.ts
@@ -71,6 +71,10 @@ export class BranchState {
 		return this.rawBranchData.is_blocked
 	}
 
+	isDeferred(): boolean {
+		return !!this.rawBranchData.deferred_cl
+	}
+
 	isActive() {
 		return this.rawBranchData.is_active
 	}
diff --git a/Engine/Extras/RoboMerge/v3/functional_tests/framework/src/framework.ts b/Engine/Extras/RoboMerge/v3/functional_tests/framework/src/framework.ts
--- a/Engine/Extras/RoboMerge/v3/functional_tests/framework/src/framework.ts
+++ b/Engine/Extras/RoboMerge/v3/functional_tests/framework/src/framework.ts
@@ -415,6 +415,14 @@ export abstract class FunctionalTest {
 		return json.branch
 	}
 
+	async ensureDeferred(sourceStream: string) {
+		this.info(`Ensuring ${sourceStream} is deferred`)
+		const branchState = await this.getBranchState(sourceStream)
+		if (!branchState.isDeferred()) {
+			throw new Error(`${sourceStream} is not deferred (blockage.type !== 'deferred')!`)
+		}
+	}
+
 	async ensureBlocked(sourceStream: string, targetStream?: string) {
 		if (targetStream) {
 			const edgeDisplayName = `${sourceStream} -> ${targetStream}`
@@ -944,7 +952,7 @@ export abstract class FunctionalTest {
 				return
 			}
 			const state = new BranchState(json)
-			if (!state.isBlocked()) {
+			if (!state.isBlocked() && !state.isDeferred()) {
 				unblockedBranchStates.set(state.name, state)
 			}
 		}))
diff --git a/Engine/Extras/RoboMerge/v3/functional_tests/framework/src/tests.ts b/Engine/Extras/RoboMerge/v3/functional_tests/framework/src/tests.ts
--- a/Engine/Extras/RoboMerge/v3/functional_tests/framework/src/tests.ts
+++ b/Engine/Extras/RoboMerge/v3/functional_tests/framework/src/tests.ts
@@ -44,6 +44,7 @@ import { StompWithAdd } from './tests/stomp-with-add'
 import { StompWithBranch } from './tests/stomp-with-branch'
 import { StompWithDelete } from './tests/stomp-with-delete'
 import { SyntaxErrorOnUnknownBranch } from './tests/syntax-error-on-unknown-branch'
+import { SyntaxErrDeferDefaultBot, SyntaxErrDeferNonDefaultBot } from './tests/syntax-error-deferred-multibot'
 import { TestChain } from './tests/test-chain'
 import { TestEdgeBlockAssetsReconsider } from './tests/test-edge-block-assets-reconsider'
 import { TestEdgeGate } from './tests/test-edge-gate'
@@ -233,7 +234,10 @@ async function go() {
 		new SkipIfUnchanged(p4),
 		new TestEdgeBlockAssetsReconsider(p4),
 		new Approval(p4),
-		new ApprovalNonBlocking(p4)  //55
+		new ApprovalNonBlocking(p4),  //55
+
+		new SyntaxErrDeferDefaultBot(p4),
+		new SyntaxErrDeferNonDefaultBot(p4)
 
 	]
 
diff --git a/Engine/Extras/RoboMerge/v3/functional_tests/framework/src/tests/syntax-error-deferred-multibot.ts b/Engine/Extras/RoboMerge/v3/functional_tests/framework/src/tests/syntax-error-deferred-multibot.ts
new file mode 100644
--- /dev/null
+++ b/Engine/Extras/RoboMerge/v3/functional_tests/framework/src/tests/syntax-error-deferred-multibot.ts
@@ -0,0 +1,105 @@
+// Copyright Epic Games, Inc. All Rights Reserved.
+
+import { FunctionalTest, P4Client, P4Util, RobomergeBranchSpec, Stream } from '../framework'
+import { Perforce } from '../test-perforce'
+
+const DEPOT_NAME = 'SyntaxErrDefer'
+const TEXT_FILENAME = 'test.txt'
+
+const streams: Stream[] = [
+	{name: 'Main', streamType: 'mainline', depotName: DEPOT_NAME},
+	{name: 'Dev',  streamType: 'development', parent: 'Main', depotName: DEPOT_NAME}
+]
+
+let streamsPromise: Promise<void> | null = null
+function createStreams(p4: Perforce, test: FunctionalTest) {
+	if (streamsPromise) {
+		return streamsPromise.then(() =>
+			Promise.all(['Main', 'Dev'].map(s => {
+				const client = test.p4Client('testuser1', s, DEPOT_NAME)
+				return client.create(P4Util.specForClient(client))
+			}))
+		)
+	}
+
+	streamsPromise = (async () => {
+		await p4.depot('stream', test.depotSpec(DEPOT_NAME))
+		await test.createStreamsAndWorkspaces(streams, DEPOT_NAME)
+
+		const mainClient = test.getClient('Main', 'testuser1', DEPOT_NAME)
+		await P4Util.addFileAndSubmit(mainClient, TEXT_FILENAME, 'Initial content')
+
+		await p4.populate(test.getStreamPath('Dev', DEPOT_NAME), 'Initial branch of files from Main')
+	})()
+
+	return streamsPromise
+}
+
+// Default bot: default bot #ROBOMERGE directives are processed here.
+// Submitting '#robomerge unknownNode' triggers a syntax error and blocks Main.
+export class SyntaxErrDeferDefaultBot extends FunctionalTest {
+	private mainClient: P4Client
+
+	async setup() {
+		await createStreams(this.p4, this)
+		this.mainClient = this.getClient('Main', 'testuser1', DEPOT_NAME)
+	}
+
+	async run() {
+		// Bare #robomerge unknownNode — unknown to this bot, triggers a syntax error
+		await P4Util.editFileAndSubmit(this.mainClient, TEXT_FILENAME, 'Change', 'unknownNode')
+	}
+
+	verify() {
+		return this.ensureBlocked('Main')
+	}
+
+	getBranches(): RobomergeBranchSpec[] {
+		return [{
+			streamDepot: DEPOT_NAME,
+			name: this.fullBranchName('Main'),
+			streamName: 'Main',
+			flowsTo: [],
+			isDefaultBot: true
+		}]
+	}
+
+	allowSyntaxErrors() { return true }
+}
+
+// Non-default bot: default bot #ROBOMERGE directives are ignored, but should DEFER
+// while the default bot has a syntax error on the same CL.
+export class SyntaxErrDeferNonDefaultBot extends FunctionalTest {
+	async setup() {
+		await createStreams(this.p4, this)
+	}
+
+	async run() {
+		// CL was already submitted by SyntaxErrDeferDefaultBot.run()
+	}
+
+	async verify() {
+		return Promise.all([
+			this.ensureDeferred('Main'),                                         // deferred_cl must be set
+			this.checkHeadRevision('Dev', TEXT_FILENAME, 1, DEPOT_NAME)          // Dev must NOT be integrated
+		])
+	}
+
+	getBranches(): RobomergeBranchSpec[] {
+		return [
+			{
+				streamDepot: DEPOT_NAME,
+				name: this.fullBranchName('Main'),
+				streamName: 'Main',
+				flowsTo: [this.fullBranchName('Dev')],
+				isDefaultBot: false	// must be explicit: bot-level config has isDefaultBot:true
+			},
+			{
+				streamDepot: DEPOT_NAME,
+				name: this.fullBranchName('Dev'),
+				streamName: 'Dev',
+				flowsTo: []
+			}
+		]
+	}
+}
diff --git a/Engine/Extras/RoboMerge/v3/public/css/table.css b/Engine/Extras/RoboMerge/v3/public/css/table.css
--- a/Engine/Extras/RoboMerge/v3/public/css/table.css
+++ b/Engine/Extras/RoboMerge/v3/public/css/table.css
@@ -216,6 +216,9 @@
   .info-block.conflict {
     background-color: #eeb8b8;
   }
+  .info-block.deferred {
+    background-color: #d4edff;
+  }
 
   .namecell .pause-info {
     white-space: pre-wrap;
diff --git a/Engine/Extras/RoboMerge/v3/public/js/boilerplate.js b/Engine/Extras/RoboMerge/v3/public/js/boilerplate.js
--- a/Engine/Extras/RoboMerge/v3/public/js/boilerplate.js
+++ b/Engine/Extras/RoboMerge/v3/public/js/boilerplate.js
@@ -1145,6 +1145,21 @@ function createPauseDivs(data, conflict) {
 		)
 	}
 
+	// Display deferral info if applicable
+	if (data.deferred_cl) {
+		const cl = data.deferred_cl
+		divs.push($('<div class="info-block deferred">')
+			.append(
+				$('<p>').html(
+					`<h6>Deferred on CL ${makeClMenuButton(cl, data.swarmURL)}</h6>` +
+					`This bot is waiting for CL ${cl} to be evaluated by another bot ` +
+					`before proceeding.<br>If this persists, check whether CL ${cl} ` +
+					`has a syntax error.`
+				)
+			)
+		)
+	}
+
 	// Display blockage data if applicable
 	if (data.is_blocked) {
 		const changelist = conflict && conflict.cl ? conflict.cl :
@@ -1231,7 +1246,7 @@ function postRenderNameCell_Node(nameCell, nodeData, includeCollapseControls, co
 	}
 
 	// Append Pause Info
-	if (!nodeData.is_available) {
+	if (!nodeData.is_available || nodeData.deferred_cl) {
 		nameCell.append(createPauseDivs(nodeData, conflict))
 	}
 
@@ -1290,11 +1305,13 @@ function preRenderStatusCell_Edge(nodeData, edgeData) {
 	else if (!edgeData.is_paused) {
 		// If we're not paused, but the node is paused or blocked, display a message informing users
 		// that we'll not be receiving any work.
-		if (nodeData.is_paused || nodeData.is_blocked || nodeData.retry_cl) {
+		if (nodeData.is_paused || nodeData.is_blocked || nodeData.retry_cl || nodeData.deferred_cl) {
+			const nodeLabel = nodeData.is_paused ? 'PAUSED' : (nodeData.deferred_cl && !nodeData.is_blocked) ? 'DEFERRED' : 'BLOCKED'
+			const nodeColor = (nodeData.deferred_cl && !nodeData.is_blocked && !nodeData.is_paused) ? 'LIGHTBLUE' : 'ORANGE'
 			$('<div>')
-				.text('NODE ' + (nodeData.is_paused ? 'PAUSED' : 'BLOCKED'))
+				.text('NODE ' + nodeLabel)
 				.addClass('important-status')
-				.css('color', 'ORANGE')
+				.css('color', nodeColor)
 				.prependTo(statusCell)
 		}
 		else if (edgeData.gateClosedMessage) {
@@ -1365,7 +1382,11 @@ function renderStatusCell_Common(statusCell, data) {
 
 	let text = null, color = 'GREEN'
 	// If the node is currently processing (active) display as such
-	if (data.retry_cl) {
+	if (data.deferred_cl) {
+		text = 'DEFERRED'
+		color = 'LIGHTBLUE'
+	}
+	else if (data.retry_cl) {
 		text = 'RETRYING'
 		color = 'ORANGE'
 	}
diff --git a/Engine/Extras/RoboMerge/v3/src/robo/branch-interfaces.ts b/Engine/Extras/RoboMerge/v3/src/robo/branch-interfaces.ts
--- a/Engine/Extras/RoboMerge/v3/src/robo/branch-interfaces.ts
+++ b/Engine/Extras/RoboMerge/v3/src/robo/branch-interfaces.ts
@@ -186,6 +186,7 @@ export interface ChangeInfo extends TargetInfo {
 	source: string
 	description: string
 
+	defer?: boolean
 	propagatingNullMerge: boolean
 	forceCreateAShelf: boolean
 	edgeServerToHostShelf?: EdgeServer
diff --git a/Engine/Extras/RoboMerge/v3/src/robo/edgebot.ts b/Engine/Extras/RoboMerge/v3/src/robo/edgebot.ts
--- a/Engine/Extras/RoboMerge/v3/src/robo/edgebot.ts
+++ b/Engine/Extras/RoboMerge/v3/src/robo/edgebot.ts
@@ -76,6 +76,7 @@ export class EdgeMergeResults {
 	edges = new Map<string, EdgeIntegrationDetails>()
 	// accumulated errors for all edges
 	errors: string[] = []
+	deferred = false
 }
 
 
diff --git a/Engine/Extras/RoboMerge/v3/src/robo/nodebot.ts b/Engine/Extras/RoboMerge/v3/src/robo/nodebot.ts
--- a/Engine/Extras/RoboMerge/v3/src/robo/nodebot.ts
+++ b/Engine/Extras/RoboMerge/v3/src/robo/nodebot.ts
@@ -52,7 +52,7 @@ function isStompableNonBinary(filename: string) {
 }
 
 
-type InfoResult = { info?: ChangeInfo, errors: any[] }
+type InfoResult = { info?: ChangeInfo, errors: any[], defer?: boolean }
 
 
 export class NodeBot extends PerforceStatefulBot implements NodeBotInterface {
@@ -82,7 +82,7 @@ export class NodeBot extends PerforceStatefulBot implements NodeBotInterface {
 	 */
 	private ticksSinceLastNewP4Commit = 0
 	private skipTickCounter = 0
-
+	private deferredCl: number | null = null
 
 	// required by Bot interface
 	public isActive = false
@@ -1335,6 +1335,10 @@ export class NodeBot extends PerforceStatefulBot implements NodeBotInterface {
 			status.lastBlockage = this.lastBlockage
 		}
 
+		if (this.deferredCl !== null) {
+			status.deferred_cl = this.deferredCl
+		}
+
 		status.conflicts = [] as any[]
 		await this.conflicts.applyStatus(status.conflicts, this.slackMessages)
 
@@ -1654,11 +1658,19 @@ export class NodeBot extends PerforceStatefulBot implements NodeBotInterface {
 
 			const changeResult = await this._processAndMergeCl(availableEdges, change, false)
 
-			// Exit immediately on syntax errors
+			if (changeResult.deferred) {
+				this.deferredCl = change.change
+				return
+			}
+
+			// Exit immediately on syntax errors or merge conflicts
 			if (this.isBlocked) {
+				this.deferredCl = null
 				return
 			}
 
+			this.deferredCl = null
+
 			if (change.change > this.lastCl) {
 				this.lastCl = change.change
 			}
@@ -1692,6 +1704,11 @@ export class NodeBot extends PerforceStatefulBot implements NodeBotInterface {
 		optOwnerOverride?: string, optWorkspaceOverride?: RoboWorkspace, optTargetBranch?: Branch | null
 	) : Promise<EdgeMergeResults> {
 		const result = await this._createChangeInfo(change, availableEdges, optOwnerOverride, optWorkspaceOverride, optTargetBranch)
+		if (result.defer) {
+			const emResult = new EdgeMergeResults()
+			emResult.deferred = true
+			return emResult
+		}
 		if (!result.info) {
 			const emResult = new EdgeMergeResults()
 			emResult.errors = result.errors
@@ -1753,6 +1770,10 @@ export class NodeBot extends PerforceStatefulBot implements NodeBotInterface {
 		// parse the change
 		const info : ChangeInfo = await this.parseChange(change, optTargetBranch)
 
+		if (info.defer) {
+			return { errors: [], defer: true }
+		}
+
 		if (optOwnerOverride) {
 			info.owner = optOwnerOverride
 		}
@@ -2019,6 +2040,19 @@ export class NodeBot extends PerforceStatefulBot implements NodeBotInterface {
 			info.owner = parsedLines.owner
 		}
 
+		// For non-default bots: if default bot directives were skipped, gate on default bot state
+		if (!this.branch.isDefaultBot && parsedLines.hasDefaultBotDirective && !commandOverride) {
+			const state = this._checkDefaultBotState(change.change)
+			if (state === 'no-default-bot') {
+				info.errors = [`Default bot #ROBOMERGE directive found but this node has no default bot — use #ROBOMERGE[${this.graphBotName}] to target this bot explicitly`]
+				return info
+			}
+			if (state === 'defer') {
+				info.defer = true
+				return info
+			}
+		}
+
 		if (parsedLines.errors.length > 0) {
 			info.errors = parsedLines.errors
 			return info
@@ -2238,6 +2272,37 @@ export class NodeBot extends PerforceStatefulBot implements NodeBotInterface {
 		this.eventTriggers.reportBlockage(blockage, true)
 	}
 
+	private _checkDefaultBotState(cl: number): 'defer' | 'no-default-bot' | 'proceed' {
+		// Find default bots monitoring the same stream as this non-default bot.
+		// We cannot rely on 'general'-flagged edges because a default bot with flowsTo:[]
+		// creates no edges at all, so we iterate all known branch graphs instead.
+		const myStream = this.branch.stream
+		const seenGraphs = new Set<BranchGraphInterface>()
+		let foundDefaultBot = false
+
+		for (const branchGraph of this.ubergraph.graph.branchGraphAliases.values()) {
+			if (branchGraph === this.branchGraph) continue     // skip our own bot
+			if (seenGraphs.has(branchGraph)) continue          // aliases may point to same graph
+			seenGraphs.add(branchGraph)
+
+			for (const branch of branchGraph.branches) {
+				if (branch.stream !== myStream || !branch.isDefaultBot) continue
+
+				foundDefaultBot = true
+				const defaultNodeBot = branch.bot
+				if (!defaultNodeBot) continue
+
+				// Default bot hasn't processed this CL yet (including if blocked on a syntax error) — defer
+				if (defaultNodeBot.lastCl < cl) {
+					this.nodeBotLogger.info(`CL ${cl}: deferring — default bot (${branchGraph.botname}) has not yet processed this CL (lastCl=${defaultNodeBot.lastCl})`)
+					return 'defer'
+				}
+			}
+		}
+
+		return foundDefaultBot ? 'proceed' : 'no-default-bot'
+	}
+
 	private async handleSyntaxError(originalCommitMessage: string, change: ChangeInfo) {
 		const shortMessage = `Syntax error in CL ${change.cl}`
 		let message = `CL ${change.cl} in ${this.branch.rootPath} contained one or more RoboMerge syntax errors (see below).\n\n`
diff --git a/Engine/Extras/RoboMerge/v3/src/robo/status-types.ts b/Engine/Extras/RoboMerge/v3/src/robo/status-types.ts
--- a/Engine/Extras/RoboMerge/v3/src/robo/status-types.ts
+++ b/Engine/Extras/RoboMerge/v3/src/robo/status-types.ts
@@ -103,6 +103,7 @@ export type BotStatusFields = Partial<PauseStatusFields> & {
 	lastBlockage?: number
 
 	retry_cl?: number
+	deferred_cl?: number
 }
 
 export type EdgeStatusFields = BotStatusFields & {
diff --git a/Engine/Extras/RoboMerge/v3/src/robo/targets.ts b/Engine/Extras/RoboMerge/v3/src/robo/targets.ts
--- a/Engine/Extras/RoboMerge/v3/src/robo/targets.ts
+++ b/Engine/Extras/RoboMerge/v3/src/robo/targets.ts
@@ -192,6 +192,8 @@ export class DescriptionParser {
 
 	errors: string[] = []
 
+	hasDefaultBotDirective = false
+
 	constructor(
 		private isDefaultBot: boolean,
 		private graphBotName: string,
@@ -211,6 +213,7 @@ export class DescriptionParser {
 		this.otherBotArguments = other.otherBotArguments
 
 		this.errors = other.errors
+		this.hasDefaultBotDirective = other.hasDefaultBotDirective
 	}
 
 	append(other: DescriptionParser) {
@@ -220,6 +223,7 @@ export class DescriptionParser {
 		this.otherBotArguments = [...this.otherBotArguments, ...other.otherBotArguments]
 
 		this.errors = [...this.errors, ...other.errors]
+		this.hasDefaultBotDirective = this.hasDefaultBotDirective || other.hasDefaultBotDirective
 	}
 
 	private processNonRoboToken(token: Token_Other) {
@@ -328,6 +332,7 @@ export class DescriptionParser {
 						this.useDefaultFlow = false
 					}
 					else {
+						this.hasDefaultBotDirective = true
 						continue
 					}
 				}
