Ensure the acknowledge interactive button is available for syntax errors

diff --git a/Engine/Extras/RoboMerge/v3/src/robo/notifications.ts b/Engine/Extras/RoboMerge/v3/src/robo/notifications.ts
--- a/Engine/Extras/RoboMerge/v3/src/robo/notifications.ts
+++ b/Engine/Extras/RoboMerge/v3/src/robo/notifications.ts
@@ -611,7 +611,7 @@ export class BotNotifications implements BotEventHandler {
 		}														
 
 		// Look up blockage URLs up front — needed for both channel message buttons and DM
-		const blockageUrls = targetBranch ? this.blockageUrlGenerator(blockage) : null
+		const blockageUrls = this.blockageUrlGenerator(blockage)
 
 		let message = this.makeSlackChannelMessage(
 			`${sourceBranch.displayName} blocked! (${issue})`,
@@ -627,7 +627,7 @@ export class BotNotifications implements BotEventHandler {
 			message.fields?.push({title: 'Shelved Change', short: true, value: makeClLink(blockage.approval.shelfCl)})
 		}
 
-		let messagesToPost = []
+		let messagesToPost: { message: SlackMessage, reply?: boolean }[] = []
 		const channelsToPostTo = this.getChannelsToPostTo(sourceBranch, targetBranch?.upperName)
 		let usersToInvite: Set<string> = new Set()
 
@@ -641,14 +641,8 @@ export class BotNotifications implements BotEventHandler {
 		}
 		else {
 			messagesToPost.push({ message })
-
-			const syntaxErrorMessage: SlackMessage = {
-				text: blockage.failure.description,
-				style: SlackMessageStyles.DANGER,
-				channel: message.channel,
-				mrkdwn: false
-			}
-			messagesToPost.push({ message: syntaxErrorMessage, reply: true})		
+			// Syntax error description is posted separately after the button (see below) so the
+			// button reply lands first and stays pinned at the top as new status replies accumulate.
 		}
 
 		let threadLinks = []
@@ -730,6 +724,54 @@ export class BotNotifications implements BotEventHandler {
 				this.slackMessages.postDM(userEmail, changeInfo.cl, targetBranch, dm)
 			}
 		}
+		else {
+			// Syntax error — button first (tracked, stays pinned), then description below it
+			if (blockageUrls?.params) {
+				const actionsBlock = this._makeActionsBlock(blockageUrls, '')
+				if (actionsBlock.elements.length > 0) {
+					for (const channel of channelsToPostTo) {
+						if (isNew) {
+							await this.slackMessages.postButtonReply(changeInfo.cl, branchArg, {
+								text: '', channel, mrkdwn: false, blocks: [actionsBlock]
+							})
+						}
+						else {
+							await this.slackMessages.updateButtonReply(changeInfo.cl, branchArg, {
+								text: '', channel, mrkdwn: false, blocks: [actionsBlock]
+							})
+						}
+					}
+				}
+			}
+
+			// Post syntax error description after the button so new status replies accumulate below
+			for (const channel of channelsToPostTo) {
+				await this.slackMessages.postReply(changeInfo.cl, branchArg, {
+					text: blockage.failure.description,
+					style: SlackMessageStyles.DANGER,
+					channel,
+					mrkdwn: false
+				})
+			}
+
+			// DM the owner with the error and an Acknowledge button
+			if (!isBotUser && userEmail) {
+				const dmText = `Your change (${makeClLink(changeInfo.source_cl)}) in *${sourceBranch.displayName}* ` +
+					`has a RoboMerge syntax error:\n\n` +
+					'`' + blockage.failure.description + '`\n\n' +
+					'*Please fix the syntax error and re-submit to unblock RoboMerge.*'
+				const dm: SlackMessage = {
+					title: `Syntax error in ${sourceBranch.displayName}`,
+					text: dmText,
+					channel: '',
+					mrkdwn: true
+				}
+				if (blockageUrls?.params) {
+					dm.blocks = [this._makeActionsBlock(blockageUrls, '')]
+				}
+				this.slackMessages.postDM(userEmail, changeInfo.cl, branchArg, dm)
+			}
+		}
 	}
 
 	private _makeActionsBlock(conflictUrls: BlockageNodeOpUrls, targetBranch: string): SlackActionsBlock {
