Minor cleanup

diff --git a/Engine/Extras/RoboMerge/v3/src/robo/notifications.ts b/Engine/Extras/RoboMerge/v3/src/robo/notifications.ts
--- a/Engine/Extras/RoboMerge/v3/src/robo/notifications.ts
+++ b/Engine/Extras/RoboMerge/v3/src/robo/notifications.ts
@@ -636,14 +636,7 @@ export class BotNotifications implements BotEventHandler {
 			usersToInvite.add(userEmail)
 		}
 
-		if (targetBranch) {
-			messagesToPost.push({ message })
-		}
-		else {
-			messagesToPost.push({ message })
-			// Syntax error description is posted separately after the button (see below) so the
-			// button reply lands first and stays pinned at the top as new status replies accumulate.
-		}
+		messagesToPost.push({ message })
 
 		let threadLinks = []
 
