Fix syntax error double posting

diff --git a/Engine/Extras/RoboMerge/v3/src/robo/notifications.ts b/Engine/Extras/RoboMerge/v3/src/robo/notifications.ts
--- a/Engine/Extras/RoboMerge/v3/src/robo/notifications.ts
+++ b/Engine/Extras/RoboMerge/v3/src/robo/notifications.ts
@@ -738,13 +738,15 @@ export class BotNotifications implements BotEventHandler {
 			}
 
 			// Post syntax error description after the button so new status replies accumulate below
-			for (const channel of channelsToPostTo) {
-				await this.slackMessages.postReply(changeInfo.cl, branchArg, {
-					text: blockage.failure.description,
-					style: SlackMessageStyles.DANGER,
-					channel,
-					mrkdwn: false
-				})
+			if (isNew) {
+				for (const channel of channelsToPostTo) {
+					await this.slackMessages.postReply(changeInfo.cl, branchArg, {
+						text: blockage.failure.description,
+						style: SlackMessageStyles.DANGER,
+						channel,
+						mrkdwn: false
+					})
+				}
 			}
 
 			// DM the owner with the error and an Acknowledge button
