Fix workspace identification not matching the correct stream in cases where one stream is the prefix of another stream name

diff --git a/Engine/Extras/RoboMerge/v3/src/common/p4util.ts b/Engine/Extras/RoboMerge/v3/src/common/p4util.ts
--- a/Engine/Extras/RoboMerge/v3/src/common/p4util.ts
+++ b/Engine/Extras/RoboMerge/v3/src/common/p4util.ts
@@ -103,8 +103,9 @@ export async function chooseBestWorkspaceForUser(p4: PerforceContext, user: stri
 			// find the stream with the closest match
 			let target_match = 0
 			for (let def of workspaces) {
-				const matchlen = def.Stream ? matchPrefix(def.Stream.toLowerCase(), branch_stream) : 0
-				if (matchlen == branch_stream.length) {
+				const defStream = def.Stream ? def.Stream.toLowerCase() : ''
+				const matchlen = def.Stream ? matchPrefix(defStream, branch_stream) : 0
+				if (matchlen == branch_stream.length && matchlen == defStream.length) {
 					if (await p4.getWorkspaceEdgeServer(def.client) !== false) {
 						return def
 					}
