Clear buttons out of DMs

diff --git a/Engine/Extras/RoboMerge/v3/src/robo/notifications.ts b/Engine/Extras/RoboMerge/v3/src/robo/notifications.ts
--- a/Engine/Extras/RoboMerge/v3/src/robo/notifications.ts
+++ b/Engine/Extras/RoboMerge/v3/src/robo/notifications.ts
@@ -308,6 +308,17 @@ export class SlackMessages {
 				msg.buttonReplyTimestamp = undefined
 				msg.update()
 			}
+			// Also clear buttons embedded directly in DM message bodies
+			if (msg.messageOpts.blocks && msg.messageOpts.blocks.length > 0) {
+				try {
+					await this.slack.update(msg.timestamp, { ...msg.messageOpts, blocks: [] })
+				}
+				catch (err) {
+					this.smLogger.error(`Error clearing blocks from message: ${err}`)
+				}
+				msg.messageOpts.blocks = []
+				msg.update()
+			}
 		}
 	}
 
