Create mechanism to specify a future edge should be shelved and provided to the user for a manual merge.
Using ~ before the stream name will indicate that when that edge merge occurs to treat it as manual.

diff --git a/Engine/Extras/RoboMerge/v3/functional_tests/framework/src/tests.ts b/Engine/Extras/RoboMerge/v3/functional_tests/framework/src/tests.ts
--- a/Engine/Extras/RoboMerge/v3/functional_tests/framework/src/tests.ts
+++ b/Engine/Extras/RoboMerge/v3/functional_tests/framework/src/tests.ts
@@ -45,6 +45,7 @@ import { StompWithBranch } from './tests/stomp-with-branch'
 import { StompWithDelete } from './tests/stomp-with-delete'
 import { SyntaxErrorOnUnknownBranch } from './tests/syntax-error-on-unknown-branch'
 import { SyntaxErrDeferDefaultBot, SyntaxErrDeferNonDefaultBot } from './tests/syntax-error-deferred-multibot'
+import { ManualEdgeMerges } from './tests/manual-edge-merges'
 import { TestChain } from './tests/test-chain'
 import { TestEdgeBlockAssetsReconsider } from './tests/test-edge-block-assets-reconsider'
 import { TestEdgeGate } from './tests/test-edge-gate'
@@ -237,7 +238,8 @@ async function go() {
 		new ApprovalNonBlocking(p4),  //55
 
 		new SyntaxErrDeferDefaultBot(p4),
-		new SyntaxErrDeferNonDefaultBot(p4)
+		new SyntaxErrDeferNonDefaultBot(p4),
+		new ManualEdgeMerges(p4)
 
 	]
 
diff --git a/Engine/Extras/RoboMerge/v3/functional_tests/framework/src/tests/cross-bot.ts b/Engine/Extras/RoboMerge/v3/functional_tests/framework/src/tests/cross-bot.ts
--- a/Engine/Extras/RoboMerge/v3/functional_tests/framework/src/tests/cross-bot.ts
+++ b/Engine/Extras/RoboMerge/v3/functional_tests/framework/src/tests/cross-bot.ts
@@ -6,7 +6,6 @@ import { Perforce } from '../test-perforce'
 const DEPOT_NAME = 'CrossBot'
 
 
-// @todo make an auto merge target and make sure ~blah works cross bot
 const streams: Stream[] = [
 	{name: 'Main', streamType: 'mainline', depotName: DEPOT_NAME},
 	{name: 'Release', streamType: 'release', parent: 'Main', depotName: DEPOT_NAME},
@@ -27,7 +26,7 @@ function createStreams(p4: Perforce, test: FunctionalTest) {
 	if (streamsPromise) {
 		// add clients to the isRoboMergeIdle list
 		return streamsPromise.then(
-			() => Promise.all(['Main', 'RES', 'FNM', 'Plus', 'R26', 'R25', 'R25Dev'].map(s => {
+			() => Promise.all(['Main', 'Dev', 'RES', 'FNM', 'Plus', 'R26', 'R25', 'R25Dev'].map(s => {
 				const client = test.p4Client('testuser1', s, DEPOT_NAME)
 				return client.create(P4Util.specForClient(client))
 			}))
@@ -40,7 +39,7 @@ function createStreams(p4: Perforce, test: FunctionalTest) {
 				await test.createStreamsAndWorkspaces(streams, DEPOT_NAME)
 
 				const mainClient = test.getClient('Main', 'testuser1', DEPOT_NAME)
-				for (const n of [1, 2, 3, 4, 5, 6, 7, 8]) {
+				for (const n of [1, 2, 3, 4, 5, 6, 7, 8, 9]) {
 					await P4Util.addFileAndSubmit(mainClient, `test${n}.txt`, 'Initial content')
 				}
 
@@ -86,6 +85,10 @@ export class CrossBotTest extends FunctionalTest {
 		await P4Util.editFile(releaseClient, 'test3.txt', 'Initial content\n\nMergeable line')
 		await P4Util.submit(releaseClient, `Edit\n#robomerge[${crossBot2!.botName}-alias] dev`)
 
+		// ~dev: cross-bot manual merge — bot2 should shelf at Dev instead of auto-merging
+		await P4Util.editFile(releaseClient, 'test9.txt', 'Initial content\n\nMergeable line')
+		await P4Util.submit(releaseClient, `Edit\n#robomerge[${crossBot2!.botName}] ~dev`)
+
 // not testing invalid bot here - more important to get default syntax error testing
 		// await P4Util.editFile(releaseClient, 'test2.txt', 'Initial content\n\nMergeable line')
 		// await P4Util.submit(releaseClient, `Edit\n#robomerge[InvalidBot] dev-alias`)
@@ -133,15 +136,23 @@ export class CrossBotTest2 extends FunctionalTest {
 
 	async run() {}
 
-	verify() {
-		return Promise.all([
+	async verify() {
+		await Promise.all([
 			this.ensureNotBlocked('Main', 'Dev'),
 			this.checkHeadRevision('Release', 'test1.txt', 2, DEPOT_NAME),
 			// this.checkHeadRevision('Release', 'test2.txt', 2, DEPOT_NAME),
 			this.checkHeadRevision('Dev', 'test1.txt', 2, DEPOT_NAME),
 			// this.checkHeadRevision('Dev', 'test2.txt', 1, DEPOT_NAME), // unrelated bot name
-			this.checkHeadRevision('Dev', 'test3.txt', 2, DEPOT_NAME)  // alias
+			this.checkHeadRevision('Dev', 'test3.txt', 2, DEPOT_NAME), // alias
+			this.checkHeadRevision('Dev', 'test9.txt', 1, DEPOT_NAME)  // ~dev: shelved, not submitted
 		])
+
+		// Confirm a pending/shelved CL exists in Dev workspace (from ~dev cross-bot manual)
+		const devClient = this.getClient('Dev', 'testuser1', DEPOT_NAME)
+		const pendingChanges = await devClient.changes(1, true)
+		if (pendingChanges.length === 0) {
+			throw new Error('Expected a shelved CL in Dev workspace from cross-bot ~dev but found none')
+		}
 	}
 
 	getBranches() {
diff --git a/Engine/Extras/RoboMerge/v3/functional_tests/framework/src/tests/manual-edge-merges.ts b/Engine/Extras/RoboMerge/v3/functional_tests/framework/src/tests/manual-edge-merges.ts
new file mode 100644
--- /dev/null
+++ b/Engine/Extras/RoboMerge/v3/functional_tests/framework/src/tests/manual-edge-merges.ts
@@ -0,0 +1,66 @@
+// Copyright Epic Games, Inc. All Rights Reserved.
+import { FunctionalTest, P4Client, P4Util, Stream } from '../framework'
+
+// Tests that ~<branchName> in a #robomerge command shelves at the correct pointinstead of auto-merging
+// Also tests that multiple ~<branchName> in a #robomerge command shelves for each
+const streams: Stream[] = [
+	{name: 'Main', streamType: 'mainline'},
+	{name: 'Dev1', streamType: 'development', parent: 'Main'},
+	{name: 'Dev2', streamType: 'development', parent: 'Main'},
+	{name: 'Dev3', streamType: 'development', parent: 'Main'},
+	{name: 'Dev1_Plus', streamType: 'development', parent: 'Main'},
+]
+
+export class ManualEdgeMerges extends FunctionalTest {
+	async setup() {
+		await this.p4.depot('stream', this.depotSpec())
+		await this.createStreamsAndWorkspaces(streams)
+
+		this.mainClient = this.getClient('Main')
+		await P4Util.addFileAndSubmit(this.mainClient, 'test.txt', 'Initial content')
+
+		await this.populateStreams(streams.slice(1))
+	}
+
+	async run() {
+		const manualMerges = streams.slice(2).map(s => `~${this.fullBranchName(s.name)}`).join(" ")
+		await P4Util.editFileAndSubmit(this.mainClient, 'test.txt', 'Main edit',
+			`${manualMerges}`)
+
+		await this.waitForRobomergeIdle()
+
+		// Confirm a pending/shelved CL exists in the Release workspace
+		await Promise.all(
+			streams.slice(2).map(async s => {
+			const releaseClient = this.getClient(s.name)
+			const pendingChanges = await releaseClient.changes(1, true)
+			if (pendingChanges.length === 0) {
+				throw new Error(`Expected a shelved CL in ${s.name} workspace but found none`)
+			}
+		}))
+	}
+
+	async verify() {
+		// Dev should have been auto-merged (head rev 2)
+		// Release should NOT have been submitted (still head rev 1)
+		await Promise.all([
+			this.checkHeadRevision('Main', 'test.txt', 2),
+			this.checkHeadRevision('Dev1', 'test.txt', 2),
+			this.checkHeadRevision('Dev2', 'test.txt', 1),
+			this.checkHeadRevision('Dev3', 'test.txt', 1),
+			this.checkHeadRevision('Dev1_Plus', 'test.txt', 1),
+		])
+	}
+
+	getBranches() {
+		return [
+			this.makeForceAllBranchDef('Main', ['Dev1', 'Dev2', 'Dev3']),
+			this.makeForceAllBranchDef('Dev1', ['Dev1_Plus']),
+			this.makeForceAllBranchDef('Dev2', []),
+			this.makeForceAllBranchDef('Dev3', []),
+			this.makeForceAllBranchDef('Dev1_Plus', [])
+		]
+	}
+
+	private mainClient: P4Client
+}
diff --git a/Engine/Extras/RoboMerge/v3/src/robo/branch-interfaces.ts b/Engine/Extras/RoboMerge/v3/src/robo/branch-interfaces.ts
--- a/Engine/Extras/RoboMerge/v3/src/robo/branch-interfaces.ts
+++ b/Engine/Extras/RoboMerge/v3/src/robo/branch-interfaces.ts
@@ -148,7 +148,7 @@ export interface Target {
 	otherBot?: string
 }
 
-export type MergeMode = 'safe' | 'normal' | 'null' | 'clobber' | 'skip'
+export type MergeMode = 'safe' | 'normal' | 'null' | 'clobber' | 'skip' | 'manual'
 export interface MergeAction {
 	branch: Branch
 	mergeMode: MergeMode
diff --git a/Engine/Extras/RoboMerge/v3/src/robo/edgebot.ts b/Engine/Extras/RoboMerge/v3/src/robo/edgebot.ts
--- a/Engine/Extras/RoboMerge/v3/src/robo/edgebot.ts
+++ b/Engine/Extras/RoboMerge/v3/src/robo/edgebot.ts
@@ -421,8 +421,9 @@ class EdgeBotImpl extends PerforceStatefulBot {
 		const thisBotMergeCommands = target.furtherMerges
 			.filter(target => !target.otherBot)
 			.map(target => (
-				target.mergeMode === 'skip' ? '-' :
-				target.mergeMode === 'null' ? '!' : '') + target.branchName)
+				target.mergeMode === 'skip'   ? '-' :
+				target.mergeMode === 'null'   ? '!' :
+				target.mergeMode === 'manual' ? '~' : '') + target.branchName)
 
 		if (flags || thisBotMergeCommands.length !== 0) {
 			const botAliases = this.sourceNode.branchGraph.config.aliases
@@ -435,7 +436,10 @@ class EdgeBotImpl extends PerforceStatefulBot {
 		for (const other of target.furtherMerges) {
 			// slight hack here: mergemode is always normal and branchName is all the original commands/flags unaltered
 			if (other.otherBot) {
-				description += `#ROBOMERGE[${other.otherBot}]: ${other.branchName}\n`
+				const otherPrefix = other.mergeMode === 'skip'   ? '-' :
+				                    other.mergeMode === 'null'   ? '!' :
+				                    other.mergeMode === 'manual' ? '~' : ''
+				description += `#ROBOMERGE[${other.otherBot}]: ${otherPrefix}${other.branchName}\n`
 			}
 		}
 
@@ -600,7 +604,7 @@ class EdgeBotImpl extends PerforceStatefulBot {
 		const doVirtualMerge = 
 			   info.forceStompChanges 
 			|| info.sendNoShelfNotification 
-		    || (!info.userRequest && !info.forceCreateAShelf && !target.flags.has('manual'))
+		    || (!info.userRequest && !info.forceCreateAShelf && target.mergeMode !== 'manual')
 
 		this.currentIntegrationStartTimestamp = Date.now()
 		await this.p4.sync(info.targetWorkspace, this.targetBranch.rootPath + '#0', {edgeServerAddress})
@@ -776,7 +780,7 @@ class EdgeBotImpl extends PerforceStatefulBot {
 		let result = await this.p4.resolve(
 			pending.change.targetWorkspace,
 			pending.newCl,
-			pending.change.forceStompChanges ? 'normal' : pending.action.mergeMode,
+			(pending.change.forceStompChanges || pending.action.mergeMode === 'manual') ? 'normal' : pending.action.mergeMode,
 			false, // detail === 'quick',
 			edgeServerAddress
 		)
@@ -805,8 +809,8 @@ class EdgeBotImpl extends PerforceStatefulBot {
 			)
 		}
 
-		if (pending.action.flags.has('manual') || pending.change.forceCreateAShelf) {
-			// the user requested manual merge
+		if (pending.action.mergeMode === 'manual' || pending.change.forceCreateAShelf) {
+			// the user requested manual merge (~BranchName, #manual flag, or forceCreateAShelf)
 			await this.shelveChangelist(pending)
 
 			if (!pending.change.sendNoShelfNotification) {
diff --git a/Engine/Extras/RoboMerge/v3/src/robo/graph.ts b/Engine/Extras/RoboMerge/v3/src/robo/graph.ts
--- a/Engine/Extras/RoboMerge/v3/src/robo/graph.ts
+++ b/Engine/Extras/RoboMerge/v3/src/robo/graph.ts
@@ -432,7 +432,7 @@ export function addBranchGraph(graph: Graph, branchGraph: BranchGraphInterface)
 	// }
 }
 
-export type MergeMode = 'safe' | 'normal' | 'null' | 'clobber' | 'skip'
+export type MergeMode = 'safe' | 'normal' | 'null' | 'clobber' | 'skip' | 'manual'
 
 class Test {
 	graph = new Graph
@@ -500,6 +500,12 @@ class Test {
 				this.targets.set(this.getNode(targetChar), 'null')
 			}
 		}
+
+		if (targetBits.length > 3) {
+			for (const targetChar of targetBits[3]) {
+				this.targets.set(this.getNode(targetChar), 'manual')
+			}
+		}
 	}
 
 	private formatNode(node: Node) {
@@ -507,7 +513,8 @@ class Test {
 		const mm = this.targets.get(node)
 		const prefix = !mm || mm === 'normal' ? '' :
 			mm === 'skip' ? '-' :
-			mm === 'null' ? '!' : '@'
+			mm === 'null' ? '!' :
+			mm === 'manual' ? '~' : '@'
 
 		return prefix + node.debugName[5]
 	}
@@ -723,10 +730,10 @@ function parseTargetsAndFlagsImpl(tokens: string[], logger: ContextualLogger, fo
 
 		let targetName = token
 		let mergeMode = forcedMode || 'normal'
-		if (prefix === '!' || prefix === '-')
+		if (prefix === '!' || prefix === '-' || prefix === '~')
 		{
 			targetName = token.substring(1)
-			mergeMode = prefix === '!' ? 'null' : 'skip'
+			mergeMode = prefix === '!' ? 'null' : prefix === '-' ? 'skip' : 'manual'
 		}
 
 		const existingMergeMode = targets.get(targetName)
diff --git a/Engine/Extras/RoboMerge/v3/src/robo/nodebot.ts b/Engine/Extras/RoboMerge/v3/src/robo/nodebot.ts
--- a/Engine/Extras/RoboMerge/v3/src/robo/nodebot.ts
+++ b/Engine/Extras/RoboMerge/v3/src/robo/nodebot.ts
@@ -1556,7 +1556,7 @@ export class NodeBot extends PerforceStatefulBot implements NodeBotInterface {
 
 		this.nodeBotLogger.info(`Sending shelf creation (shelf CL ${pendingChange.newCl}) notification to ${pendingChange.change.owner} for source CL ${pendingChange.change.source_cl}`)
 
-		let showStillBlockedMessage = pendingChange.change.forceCreateAShelf && !pendingChange.action.flags.has('manual')
+		let showStillBlockedMessage = pendingChange.change.forceCreateAShelf && pendingChange.action.mergeMode !== 'manual'
 
 		if (this.slackMessages) {
 			let dm: SlackMessage = {
@@ -1719,7 +1719,7 @@ export class NodeBot extends PerforceStatefulBot implements NodeBotInterface {
 		// targets as long as all the target workspaces are on the same edge, but it will still crash if you
 		// have multiple targets with workspaces on different edgeservers
 		const numTargets = (result.info.targets || []).length;
-		const isManualChange = numTargets > 0 ? result.info.targets![0].flags.has('manual') : false;
+		const isManualChange = numTargets > 0 ? result.info.targets![0].mergeMode === 'manual' : false;
 		if (change.forceCreateAShelf || (change.isUserRequest && !change.forceStompChanges) || isManualChange) {
 			if (!optWorkspaceOverride && numTargets == 1) {
 				const target = result.info.targets![0]
diff --git a/Engine/Extras/RoboMerge/v3/src/robo/targets.ts b/Engine/Extras/RoboMerge/v3/src/robo/targets.ts
--- a/Engine/Extras/RoboMerge/v3/src/robo/targets.ts
+++ b/Engine/Extras/RoboMerge/v3/src/robo/targets.ts
@@ -35,7 +35,7 @@ export function getIntegrationOwner(arg0: Branch | PendingChange, sourceBranch?:
 
 	// Manual requester or edge reconsider
 	if (pending) {
-		if (pending.action.flags.has('manual') ||
+		if (pending.action.mergeMode === 'manual' ||
 			pending.change.forceCreateAShelf ||
 			pending.change.userRequest == 'edge-reconsider'
 		) {
@@ -459,7 +459,7 @@ export function processOtherBotTargets(
 			continue
 		}
 
-		const target = '!-'.indexOf(arg[0]) < 0 ? arg : arg.substring(1)
+		const target = '!-~'.indexOf(arg[0]) < 0 ? arg : arg.substring(1)
 		const branch = targetBranchGraph.getBranch(target)
 		if (!branch) {
 			if (!targetBranchGraph.config.macros[target.toLowerCase()] &&
@@ -580,6 +580,10 @@ class RequestedIntegrations {
 				this.integrations.push([arg.substring(1), 'skip'])
 				break
 
+			case '~':
+				this.integrations.push([arg.substring(1), 'manual'])
+				break
+
 			case '$':
 			case '#':
 				// set the flag and continue the loop
@@ -708,6 +712,11 @@ function computeTargetsImpl(
 		.filter(([_, mergeMode]) => mergeMode === 'skip')
 		.map(([branchName, _]) => branchGraph.getBranch(branchName)!)
 
+	// handle multiple routes to manual targets
+	const manualTargets: Branch[] = [...requestedMerges]
+		.filter(([_, mergeMode]) => mergeMode === 'manual')
+		.map(([branchName, _]) => branchGraph.getBranch(branchName)!)
+		.filter(Boolean)
 
 	const merges = new Map<Branch, Branch[]>()
 	for (const [initialEdge, furtherEdges] of integrations) {
@@ -738,6 +747,17 @@ function computeTargetsImpl(
 				furtherEdgeBranches.push(skipTarget)
 			}
 		}
+		for (const manualTarget of manualTargets) {
+			if (furtherEdgeBranches.indexOf(manualTarget) < 0
+				&&
+				[...furtherEdgeBranches, targetBranch]
+					.map(b => b.forcedDownstream)
+					.filter(ds => (ds || []).map(b => branchGraph.getBranch(b)).indexOf(manualTarget) >= 0)
+					.length > 0
+				) {
+				furtherEdgeBranches.push(manualTarget)
+			}
+		}
 		merges.set(targetBranch, furtherEdgeBranches)
 	}
 
@@ -797,7 +817,7 @@ export function computeTargets(
 	const allDownstream = new Set<Branch>()
 
 	for (const [direct, indirectTargets] of merges.entries()) {
-		const mergeMode : MergeMode = targets.get(direct) || 'normal'
+		const mergeMode : MergeMode = targets.get(direct) || (flags.has('manual') ? 'manual' : 'normal')
 		if (mergeMode === 'skip') {
 			continue
 		}
