Avoid track changes crashing when trying to track a change across a uniqueBranch without a stream name

diff --git a/Engine/Extras/RoboMerge/v3/src/robo/trackchange.ts b/Engine/Extras/RoboMerge/v3/src/robo/trackchange.ts
--- a/Engine/Extras/RoboMerge/v3/src/robo/trackchange.ts
+++ b/Engine/Extras/RoboMerge/v3/src/robo/trackchange.ts
@@ -370,17 +370,19 @@ export async function trackChange(robo: RoboMerge, logger: ContextualLogger, ini
             if (entry.action != 'move/delete') {
                 for (let i=externalServerTargets.length-1; i>=0; i--) {
                     const externalTarget = externalServerTargets[i]
-                    const externalP4 = PerforceContext.getServerContext(logger, externalTarget.serverID)
-                    const streamName = getStreamFromPath(changeToConsider.node!.stream)
-                    const targetPath = entry.depotFile.replace(streamName || changeToConsider.node!.stream, externalTarget.stream)
-                    const changes = await externalP4.changes(targetPath)
-                    for (let change of changes) {
-                        if (change.desc.includes(clToConsider.cl.toString())) {
-                            const newChangeToConsider = {cl: change.change, serverId: externalTarget.serverID}
-                            changeToConsider.destCLs.push(newChangeToConsider)
-                            await gatherCLInfo(newChangeToConsider, {parentCL: clToConsider})
-                            externalServerTargets.splice(i, 1)
-                            break
+                    if (externalTarget.stream.startsWith('//')) {
+                        const externalP4 = PerforceContext.getServerContext(logger, externalTarget.serverID)
+                        const streamName = getStreamFromPath(changeToConsider.node!.stream)
+                        const targetPath = entry.depotFile.replace(streamName || changeToConsider.node!.stream, externalTarget.stream)
+                        const changes = await externalP4.changes(targetPath)
+                        for (let change of changes) {
+                            if (change.desc.includes(clToConsider.cl.toString())) {
+                                const newChangeToConsider = {cl: change.change, serverId: externalTarget.serverID}
+                                changeToConsider.destCLs.push(newChangeToConsider)
+                                await gatherCLInfo(newChangeToConsider, {parentCL: clToConsider})
+                                externalServerTargets.splice(i, 1)
+                                break
+                            }
                         }
                     }
                 }
@@ -438,23 +440,26 @@ export async function trackChange(robo: RoboMerge, logger: ContextualLogger, ini
                     const srcStream = getStreamFromPath(entry.depotFile)
                     if (srcStream) {
                         for (let i=sameServerTargets.length-1; i>=0; i--) {
-                            const targetPath = entry.depotFile.replace(srcStream, sameServerTargets[i].stream)
-                            const targetCLs = await p4.changes(targetPath, {since: clToConsider.cl})
-                            for (const targetCL of targetCLs) {
-                                if (targetCL.desc.includes(clToConsider.cl.toString())) {
-                                    const newCL = {cl: targetCL.change, serverId: clToConsider.serverId}
-                                    if (!changeToConsider.destCLs.some(cl => compareTrackedCL(cl, newCL))) {
-                                        changeToConsider.destCLs.push(newCL)
-                                        const existingChange = changes.get(newCL)
-                                        if (existingChange) {
-                                            existingChange.sourceCL = clToConsider
-                                        }
-                                        else {
-                                            await gatherCLInfo(newCL, {parentCL: clToConsider})
+                            const targetStream = sameServerTargets[i].stream
+                            if (targetStream.startsWith('//')) {
+                                const targetPath = entry.depotFile.replace(srcStream, targetStream)
+                                const targetCLs = await p4.changes(targetPath, {since: clToConsider.cl})
+                                for (const targetCL of targetCLs) {
+                                    if (targetCL.desc.includes(clToConsider.cl.toString())) {
+                                        const newCL = {cl: targetCL.change, serverId: clToConsider.serverId}
+                                        if (!changeToConsider.destCLs.some(cl => compareTrackedCL(cl, newCL))) {
+                                            changeToConsider.destCLs.push(newCL)
+                                            const existingChange = changes.get(newCL)
+                                            if (existingChange) {
+                                                existingChange.sourceCL = clToConsider
+                                            }
+                                            else {
+                                                await gatherCLInfo(newCL, {parentCL: clToConsider})
+                                            }
                                         }
+                                        sameServerTargets.splice(i, 1)
+                                        break
                                     }
-                                    sameServerTargets.splice(i, 1)
-                                    break
                                 }
                             }
                         }
