Fix crashes in async code that takes down the service not notifying alerts channel

diff --git a/Engine/Extras/RoboMerge/v3/src/common/helper.ts b/Engine/Extras/RoboMerge/v3/src/common/helper.ts
--- a/Engine/Extras/RoboMerge/v3/src/common/helper.ts
+++ b/Engine/Extras/RoboMerge/v3/src/common/helper.ts
@@ -6,15 +6,41 @@ import { ContextualLogger } from './logger';
 
 const helperLogger = new ContextualLogger('Helper')
 
+type ProcessAlertFn = (message: string) => Promise<void>
+let _processAlertFn: ProcessAlertFn | null = null
+
+export function setProcessAlertCallback(fn: ProcessAlertFn) {
+	_processAlertFn = fn
+}
+
 // fail hard on exceptions thrown in async functions
 process.on('unhandledRejection', async (reason, promise) => {
 	console.log('Encountered unhandledRejection at: ', promise, '\n\nreason:', reason)
 
+	if (_processAlertFn) {
+		const msg = reason instanceof Error ? reason.message : String(reason)
+		const stack = reason instanceof Error ? (reason.stack || '') : ''
+		await _processAlertFn(`@here Robomerge process crashed (unhandled rejection): ${msg}\n\`\`\`${stack}\`\`\``)
+	}
+
 	if (process.env.EPIC_ENV === 'prod') {
 		// Sentry integration has a unhandledRejection listener, need to give it time to send to Sentry before throwing fatal err
 		await closeSentry(helperLogger)
 	}
 	process.exit(-1) // this is unexpected but not the fault of the committed changelist
+})
+
+process.on('uncaughtException', async (err: Error) => {
+	console.log('Encountered uncaughtException:', err)
+
+	if (_processAlertFn) {
+		await _processAlertFn(`@here Robomerge process crashed (uncaught exception): ${err.message}\n\`\`\`${err.stack}\`\`\``)
+	}
+
+	if (process.env.EPIC_ENV === 'prod') {
+		await closeSentry(helperLogger)
+	}
+	process.exit(-1)
 });
 
 /**
diff --git a/Engine/Extras/RoboMerge/v3/src/robo/robo.ts b/Engine/Extras/RoboMerge/v3/src/robo/robo.ts
--- a/Engine/Extras/RoboMerge/v3/src/robo/robo.ts
+++ b/Engine/Extras/RoboMerge/v3/src/robo/robo.ts
@@ -23,7 +23,8 @@ import { NodeBot } from './nodebot';
 import { roboAnalytics, setGlobalAnalytics } from './roboanalytics';
 import { RoboServer } from './roboserver';
 import { SlackSocket } from './slack-socket';
-import { notificationsInit } from './notifications'
+import { notificationsInit, postToRobomergeAlerts } from './notifications'
+import { setProcessAlertCallback } from '../common/helper'
 import { RoboArgs } from './roboargs'
 import { settingsInit } from './settings'
 import { doDBPersistedMessagesHousekeeping } from './persistedmessages';
@@ -288,6 +289,8 @@ RoboArgs.Init(args)
 notificationsInit()
 settingsInit()
 
+setProcessAlertCallback((msg) => postToRobomergeAlerts(msg))
+
 /**
  * Connect to MongoDB.
  */
