Set a server ID to avoid startup warnings

diff --git a/Engine/Extras/RoboMerge/v3/p4docker/setup_and_run_p4d.sh b/Engine/Extras/RoboMerge/v3/p4docker/setup_and_run_p4d.sh
--- a/Engine/Extras/RoboMerge/v3/p4docker/setup_and_run_p4d.sh
+++ b/Engine/Extras/RoboMerge/v3/p4docker/setup_and_run_p4d.sh
@@ -5,5 +5,8 @@ set -e
 p4d &
 sleep 3
 
+# Set serverID to suppress topology registration warning
+p4 serverid FunctionalTestServer
+
 # Continue running perforce server -- tail the log
 tail -f $P4LOGS
\ No newline at end of file
