Remove unnecessary verbose logging

diff --git a/Engine/Extras/RoboMerge/v3/src/common/perforce.ts b/Engine/Extras/RoboMerge/v3/src/common/perforce.ts
--- a/Engine/Extras/RoboMerge/v3/src/common/perforce.ts
+++ b/Engine/Extras/RoboMerge/v3/src/common/perforce.ts
@@ -1078,7 +1078,6 @@ export class PerforceContext {
 		form += 'Description:\n\t' + this._sanitizeDescription(description);
 
 		// run the P4 change command
-		this.logger.info("Executing: 'p4 change -i' to create a new CL");
 		while (true) {
 			try {
 				const output = await this._execP4(workspace, ['change', '-i'], { stdin: form, quiet: true, serverAddress: edgeServerAddress });
