Log track change initiation

diff --git a/Engine/Extras/RoboMerge/v3/src/robo/ipc.ts b/Engine/Extras/RoboMerge/v3/src/robo/ipc.ts
--- a/Engine/Extras/RoboMerge/v3/src/robo/ipc.ts
+++ b/Engine/Extras/RoboMerge/v3/src/robo/ipc.ts
@@ -189,6 +189,8 @@ export class IPC {
 			return {statusCode: 503, message: 'Service is not ready to process request'}
 		}
 
+		this.ipcLogger.info(`${queryObj.who} initiated track changes for CL# ${trackedCL}`)
+
 		const userTags = new Set<string>()
 		for (const tag in tagsObj) {
 			userTags.add(tag)
diff --git a/Engine/Extras/RoboMerge/v3/src/robo/roboserver.ts b/Engine/Extras/RoboMerge/v3/src/robo/roboserver.ts
--- a/Engine/Extras/RoboMerge/v3/src/robo/roboserver.ts
+++ b/Engine/Extras/RoboMerge/v3/src/robo/roboserver.ts
@@ -256,17 +256,12 @@ class RoboWebApp implements AppInterface {
 			throw new Error('Secure call but no auth data?')
 		}
 
-		const queryObj: {[key: string]: string} = {}
-		for (const [key, val] of this.request.url.searchParams) {
-			queryObj[key] = val
-		}
-
 		const tagsObj: {[key: string]: boolean} = {}
 		for (const tag of this.authData.tags) {
 			tagsObj[tag] = true
 		}
 
-		return this.sendMessage('trackChange', [queryObj, tagsObj])
+		return this.sendMessage('trackChange', [this.getQueryFromSecure(), tagsObj])
 	}
 
 
