Fix race condition in cross bot tests

diff --git a/Engine/Extras/RoboMerge/v3/functional_tests/framework/src/tests/cross-bot.ts b/Engine/Extras/RoboMerge/v3/functional_tests/framework/src/tests/cross-bot.ts
--- a/Engine/Extras/RoboMerge/v3/functional_tests/framework/src/tests/cross-bot.ts
+++ b/Engine/Extras/RoboMerge/v3/functional_tests/framework/src/tests/cross-bot.ts
@@ -250,7 +250,7 @@ export class ComplexCrossBot extends ComplexBase {
 			this.makeBranch('RES', ['DEM'], []),
 			this.makeBranch('DEM', ['RES', 'FNM'], ['RES']),
 			this.makeBranch('FNM', ['DEM'], ['DEM']),
-			this.makeBranch('Plus', ['DEM'], []),
+			this.makeBranch('Plus', [], []),
 		]
 	}
 
@@ -333,7 +333,7 @@ export class ComplexCrossBot3 extends ComplexBase {
 
 	getMacros() {
 		return {
-			'nextgen': [`#robomerge[${complexBot!.botName}] FNM, -RES`]
+			'nextgen': [`#robomerge[${complexBot!.botName}] FNM`]
 		}
 	}
 }
