Try a fix for server.id warnings

diff --git a/Engine/Extras/RoboMerge/v3/p4docker/setup_and_run_p4d.sh b/Engine/Extras/RoboMerge/v3/p4docker/setup_and_run_p4d.sh
--- a/Engine/Extras/RoboMerge/v3/p4docker/setup_and_run_p4d.sh
+++ b/Engine/Extras/RoboMerge/v3/p4docker/setup_and_run_p4d.sh
@@ -1,12 +1,12 @@
 #!/bin/bash
 set -e
 
-# First start p4d and wait for it to finish starting
+# Pre-create server.id so p4d starts with the ID already set (no topology warning)
+echo "FunctionalTestServer" > /p4root/server.id
+
+# Start p4d
 p4d &
 sleep 3
 
-# Set serverID to suppress topology registration warning
-p4 serverid FunctionalTestServer
-
 # Continue running perforce server -- tail the log
 tail -f $P4LOGS
\ No newline at end of file
