Add various logging to try and diagnose slack interactive button failures

diff --git a/Engine/Extras/RoboMerge/v3/src/robo/robo.ts b/Engine/Extras/RoboMerge/v3/src/robo/robo.ts
--- a/Engine/Extras/RoboMerge/v3/src/robo/robo.ts
+++ b/Engine/Extras/RoboMerge/v3/src/robo/robo.ts
@@ -577,9 +577,10 @@ if (!args.noIPC) {
 	})
 }
 else {
-	const sendMessage = (name: string, args?: any[]) => ipc.handle({name:name, args:args}).catch(err => console.log(`IPC error! ${err.toString()}\n${err.stack}`));
+	const sendMessage = (name: string, args?: any[]) => ipc.handle({name:name, args:args}).catch(err => roboStartupLogger.error(`IPC error! ${err.toString()}\n${err.stack}`));
 
 	if (RoboArgs.hasSocketMode) {
+		roboStartupLogger.info(`Starting Slack Socket Mode (externalUrl='${args.externalUrl}', botToken=${RoboArgs.vault.slackTokens?.bot ? 'present' : 'MISSING'})`)
 		const socket = new SlackSocket(RoboArgs.vault.slackAppToken!, args.externalUrl, {
 			sendMessage,
 			botToken: RoboArgs.vault.slackTokens?.bot,
@@ -588,6 +589,9 @@ else {
 		})
 		socket.start()
 	}
+	else {
+		roboStartupLogger.info('Slack Socket Mode disabled (slackAppToken not present in vault)')
+	}
 
 	const ws = new RoboServer(args.externalUrl, sendMessage, () => roboStartupLogger.getLogTail(),
 		() => roboStartupLogger.info('Received: getLastCrash'),
diff --git a/Engine/Extras/RoboMerge/v3/src/robo/slack-actions.ts b/Engine/Extras/RoboMerge/v3/src/robo/slack-actions.ts
--- a/Engine/Extras/RoboMerge/v3/src/robo/slack-actions.ts
+++ b/Engine/Extras/RoboMerge/v3/src/robo/slack-actions.ts
@@ -27,6 +27,7 @@ export function extractHostFromSlackPayload(payload: SlackInteractivePayload): s
 }
 
 export async function dispatchSlackInteractivePayload(payload: SlackInteractivePayload, dispatchContext: SlackDispatchContext) {
+	dispatchContext.logger.info(`dispatchSlackInteractivePayload: type=${payload.type}, user=${payload.user?.name ?? 'unknown'}`)
 	if (payload.type === 'block_actions') {
 		return _handleSlackBlockAction(payload, dispatchContext)
 	}
@@ -54,16 +55,19 @@ async function _handleSlackBlockAction(payload: SlackInteractivePayload, dispatc
 	const who = payload.user.name
 
 	if (!bot || !branch || !cl) {
+		dispatchContext.logger.warn(`_handleSlackBlockAction: missing required params — bot=${bot}, branch=${branch}, cl=${cl}`)
 		return { statusCode: 400, message: 'Missing required params (bot, branch, cl)' }
 	}
 
 	const clStr = String(cl)
+	dispatchContext.logger.info(`_handleSlackBlockAction: action='${action.action_id}' bot='${bot}' branch='${branch}' cl=${cl} edge='${edge ?? 'none'}' user='${who}'`)
 
 	switch (action.action_id) {
 		case 'acknowledge': {
 			const result = edge
 				? await dispatchContext.sendMessage('doEdgeOp', [bot, branch, edge, 'acknowledge', { who, cl: clStr }])
 				: await dispatchContext.sendMessage('doNodeOp', [bot, branch, 'acknowledge', { who, cl: clStr }])
+			dispatchContext.logger.info(`_handleSlackBlockAction: acknowledge IPC result: statusCode=${result?.statusCode}, message='${result?.message}'`)
 			return _slackActionResult(result, payload, `Acknowledged by ${who}`, dispatchContext)
 		}
 
@@ -92,6 +96,7 @@ async function _handleSlackBlockAction(payload: SlackInteractivePayload, dispatc
 				return { statusCode: 400, message: 'unlockchanges requires target param' }
 			}
 			const result = await dispatchContext.sendMessage('doNodeOp', [bot, branch, 'unlockchanges', { who, cl: clStr, target }])
+			dispatchContext.logger.info(`_handleSlackBlockAction: unlockchanges IPC result: statusCode=${result?.statusCode}, message='${result?.message}'`)
 			return _slackActionResult(result, payload, `Unlock started by ${who}`, dispatchContext)
 		}
 
@@ -113,6 +118,7 @@ async function _handleSlackViewSubmission(payload: SlackInteractivePayload, disp
 		return { statusCode: 400, message: 'No view in submission payload' }
 	}
 
+	dispatchContext.logger.info(`_handleSlackViewSubmission: callback_id='${view.callback_id}' user='${payload.user?.name ?? 'unknown'}'`)
 	switch (view.callback_id) {
 		case 'skip_reason_modal':   return _handleSkipReasonSubmit(payload, view, dispatchContext)
 		case 'stomp_confirm_modal': return _handleStompConfirmSubmit(payload, view, dispatchContext)
@@ -157,6 +163,7 @@ async function _handleCreateShelfSubmit(payload: SlackInteractivePayload, view:
 		?.some((o: any) => o.value === 'acknowledge') ?? false
 
 	const result = await dispatchContext.sendMessage('doNodeOp', [bot, branch, 'create_shelf', { who, cl: String(cl), target, workspace }])
+	dispatchContext.logger.info(`_handleCreateShelfSubmit: create_shelf IPC result: statusCode=${result?.statusCode}, message='${result?.message}'`)
 	if (result && result.statusCode === 200) {
 		if (acknowledgeSelected) {
 			// Mirror the web UI: acknowledge via edge op when an edge is present
@@ -208,6 +215,7 @@ async function _handleSkipReasonSubmit(payload: SlackInteractivePayload, view: N
 	const result = await dispatchContext.sendMessage('doEdgeOp', [bot, branch, edge, 'set_last_cl', {
 		who, cl: String(cl), reason: reasonText, unblock: 'true', isSkip: 'true'
 	}])
+	dispatchContext.logger.info(`_handleSkipReasonSubmit: set_last_cl IPC result: statusCode=${result?.statusCode}, message='${result?.message}'`)
 
 	if (result && result.statusCode === 200) {
 		if (payload.response_url && dispatchContext.botToken) {
@@ -235,6 +243,7 @@ async function _handleStompConfirmSubmit(payload: SlackInteractivePayload, view:
 	const who = payload.user.name
 
 	const result = await dispatchContext.sendMessage('doNodeOp', [bot, branch, 'stompchanges', { who, cl: String(cl), target }])
+	dispatchContext.logger.info(`_handleStompConfirmSubmit: stompchanges IPC result: statusCode=${result?.statusCode}, message='${result?.message}'`)
 
 	if (result && result.statusCode === 200) {
 		return { statusCode: 200, message: JSON.stringify({ response_action: 'clear' }), headers: [['Content-Type', 'application/json']] }
diff --git a/Engine/Extras/RoboMerge/v3/src/robo/slack-socket.ts b/Engine/Extras/RoboMerge/v3/src/robo/slack-socket.ts
--- a/Engine/Extras/RoboMerge/v3/src/robo/slack-socket.ts
+++ b/Engine/Extras/RoboMerge/v3/src/robo/slack-socket.ts
@@ -13,7 +13,7 @@ export class SlackSocket {
 	constructor(
 		private appToken: string,
 		private externalUrl: string,
-		private ctx: SlackDispatchContext
+		private dispatchContext: SlackDispatchContext
 	) {}
 
 	async start() {
@@ -23,7 +23,7 @@ export class SlackSocket {
 			this.connect(wsUrl)
 		}
 		catch (err) {
-			this.ctx.logger.error(`Slack Socket Mode: failed to get WebSocket URL: ${err}`)
+			this.dispatchContext.logger.error(`Slack Socket Mode: failed to get WebSocket URL: ${err}`)
 			this.scheduleReconnect()
 		}
 	}
@@ -54,12 +54,12 @@ export class SlackSocket {
 	}
 
 	private connect(url: string) {
-		this.ctx.logger.info('Slack Socket Mode: connecting...')
+		this.dispatchContext.logger.info('Slack Socket Mode: connecting...')
 		const ws = new WebSocket(url)
 		this.ws = ws
 
 		ws.addEventListener('open', () => {
-			this.ctx.logger.info('Slack Socket Mode: connected')
+			this.dispatchContext.logger.info('Slack Socket Mode: connected')
 		})
 
 		ws.addEventListener('message', (event: MessageEvent) => {
@@ -68,13 +68,13 @@ export class SlackSocket {
 
 		ws.addEventListener('close', (event: Event) => {
 			if (this.ws !== ws) return  // already replaced
-			this.ctx.logger.warn(`Slack Socket Mode: closed (code=${(event as any).code}), reconnecting...`)
+			this.dispatchContext.logger.warn(`Slack Socket Mode: closed (code=${(event as any).code}), reconnecting...`)
 			this.ws = null
 			this.scheduleReconnect()
 		})
 
 		ws.addEventListener('error', () => {
-			this.ctx.logger.error('Slack Socket Mode: WebSocket error')
+			this.dispatchContext.logger.error('Slack Socket Mode: WebSocket error')
 		})
 	}
 
@@ -88,18 +88,19 @@ export class SlackSocket {
 		}
 
 		if (msg.type === 'hello') {
-			this.ctx.logger.verbose('Slack Socket Mode: hello received')
+			this.dispatchContext.logger.verbose('Slack Socket Mode: hello received')
 			return
 		}
 
 		if (msg.type === 'disconnect') {
-			this.ctx.logger.info('Slack Socket Mode: server requested disconnect, reconnecting...')
+			this.dispatchContext.logger.info('Slack Socket Mode: server requested disconnect, reconnecting...')
 			this.ws?.close()
 			return
 		}
 
 		if (msg.type === 'interactive') {
 			const payload = msg.payload as SlackInteractivePayload
+			this.dispatchContext.logger.info(`Slack Socket Mode: interactive event (subtype=${payload?.type ?? 'unknown'}, envelope_id=${msg.envelope_id ?? 'none'})`)
 			if (payload?.type === 'view_submission') {
 				// For view_submission, the response_action must be included in the envelope ACK.
 				// Do NOT ACK immediately — wait for the dispatch result (within Slack's 3-second window).
@@ -127,11 +128,14 @@ export class SlackSocket {
 		// Only handle actions targeted at this instance.
 		const host = extractHostFromSlackPayload(payload)
 		if (host && host !== this.externalUrl) {
+			this.dispatchContext.logger.info(`Slack Socket Mode: ignoring block_action — host mismatch (payload='${host}', local='${this.externalUrl}')`)
 			return
 		}
 
-		dispatchSlackInteractivePayload(payload, this.ctx).catch(err => {
-			this.ctx.logger.error(`Slack Socket Mode: dispatch error: ${err}`)
+		const actionId = payload.actions?.[0]?.action_id ?? 'unknown'
+		this.dispatchContext.logger.info(`Slack Socket Mode: dispatching block_action '${actionId}' (host='${host ?? 'none'}')`)
+		dispatchSlackInteractivePayload(payload, this.dispatchContext).catch(err => {
+			this.dispatchContext.logger.error(`Slack Socket Mode: dispatch error: ${err}`)
 		})
 	}
 
@@ -145,13 +149,17 @@ export class SlackSocket {
 		// Only handle submissions targeted at this instance.
 		const host = extractHostFromSlackPayload(payload)
 		if (host && host !== this.externalUrl) {
+			this.dispatchContext.logger.info(`Slack Socket Mode: ignoring view_submission — host mismatch (payload='${host}', local='${this.externalUrl}')`)
 			if (envelopeId) this.ws?.send(JSON.stringify({ envelope_id: envelopeId }))
 			return
 		}
 
+		const callbackId = payload.view?.callback_id ?? 'unknown'
+		this.dispatchContext.logger.info(`Slack Socket Mode: dispatching view_submission '${callbackId}' (host='${host ?? 'none'}')`)
+
 		let ackPayload: any = {}
 		try {
-			const result = await dispatchSlackInteractivePayload(payload, this.ctx)
+			const result = await dispatchSlackInteractivePayload(payload, this.dispatchContext)
 			// Extract the response_action JSON from the result message and include it in the ACK
 			if (result?.message && typeof result.message === 'string') {
 				try {
@@ -164,7 +172,7 @@ export class SlackSocket {
 			}
 		}
 		catch (err) {
-			this.ctx.logger.error(`Slack Socket Mode: view submission dispatch error: ${err}`)
+			this.dispatchContext.logger.error(`Slack Socket Mode: view submission dispatch error: ${err}`)
 		}
 
 		if (envelopeId) {
