Fix slack sockets not working with IPC

diff --git a/Engine/Extras/RoboMerge/v3/src/robo/robo.ts b/Engine/Extras/RoboMerge/v3/src/robo/robo.ts
--- a/Engine/Extras/RoboMerge/v3/src/robo/robo.ts
+++ b/Engine/Extras/RoboMerge/v3/src/robo/robo.ts
@@ -579,20 +579,6 @@ if (!args.noIPC) {
 else {
 	const sendMessage = (name: string, args?: any[]) => ipc.handle({name:name, args:args}).catch(err => roboStartupLogger.error(`IPC error! ${err.toString()}\n${err.stack}`));
 
-	if (RoboArgs.hasSocketMode) {
-		roboStartupLogger.info(`Starting Slack Socket Mode (externalUrl='${args.externalUrl}', botToken=${RoboArgs.vault.slackTokens?.bot ? 'present' : 'MISSING'})`)
-		const socket = new SlackSocket(RoboArgs.vault.slackAppToken!, args.externalUrl, {
-			sendMessage,
-			botToken: RoboArgs.vault.slackTokens?.bot,
-			slackDomain: args.slackDomain,
-			logger: roboStartupLogger
-		})
-		socket.start()
-	}
-	else {
-		roboStartupLogger.info('Slack Socket Mode disabled (slackAppToken not present in vault)')
-	}
-
 	const ws = new RoboServer(args.externalUrl, sendMessage, () => roboStartupLogger.getLogTail(),
 		() => roboStartupLogger.info('Received: getLastCrash'),
 		() => roboStartupLogger.info('Received: stopBot'),
@@ -690,6 +676,21 @@ async function main(logger: ContextualLogger) {
 	await init(logger);
 
 	startBots(logger);
+
+	const sendMessage = (name: string, args?: any[]) => ipc.handle({name, args}).catch((err: any) => logger.error(`IPC error! ${err.toString()}\n${err.stack}`))
+	if (RoboArgs.hasSocketMode) {
+		logger.info(`Starting Slack Socket Mode (externalUrl='${args.externalUrl}', botToken=${RoboArgs.vault.slackTokens?.bot ? 'present' : 'MISSING'})`)
+		const socket = new SlackSocket(RoboArgs.vault.slackAppToken!, args.externalUrl, {
+			sendMessage,
+			botToken: RoboArgs.vault.slackTokens?.bot,
+			slackDomain: args.slackDomain,
+			logger
+		})
+		socket.start()
+	}
+	else {
+		logger.info('Slack Socket Mode disabled (slackAppToken not present in vault)')
+	}
 }
 
 let robo: RoboMerge
