Remove verbose debugging logs

diff --git a/Engine/Extras/RoboMerge/v3/src/robo/slack-actions.ts b/Engine/Extras/RoboMerge/v3/src/robo/slack-actions.ts
--- a/Engine/Extras/RoboMerge/v3/src/robo/slack-actions.ts
+++ b/Engine/Extras/RoboMerge/v3/src/robo/slack-actions.ts
@@ -27,7 +27,6 @@ export function extractHostFromSlackPayload(payload: SlackInteractivePayload): s
 }
 
 export async function dispatchSlackInteractivePayload(payload: SlackInteractivePayload, dispatchContext: SlackDispatchContext) {
-	dispatchContext.logger.info(`dispatchSlackInteractivePayload: type=${payload.type}, user=${payload.user?.name ?? 'unknown'}`)
 	if (payload.type === 'block_actions') {
 		return _handleSlackBlockAction(payload, dispatchContext)
 	}
@@ -67,7 +66,6 @@ async function _handleSlackBlockAction(payload: SlackInteractivePayload, dispatc
 			const result = edge
 				? await dispatchContext.sendMessage('doEdgeOp', [bot, branch, edge, 'acknowledge', { who, cl: clStr }])
 				: await dispatchContext.sendMessage('doNodeOp', [bot, branch, 'acknowledge', { who, cl: clStr }])
-			dispatchContext.logger.info(`_handleSlackBlockAction: acknowledge IPC result: statusCode=${result?.statusCode}, message='${result?.message}'`)
 			return _slackActionResult(result, payload, `Acknowledged by ${who}`, dispatchContext)
 		}
 
@@ -96,7 +94,6 @@ async function _handleSlackBlockAction(payload: SlackInteractivePayload, dispatc
 				return { statusCode: 400, message: 'unlockchanges requires target param' }
 			}
 			const result = await dispatchContext.sendMessage('doNodeOp', [bot, branch, 'unlockchanges', { who, cl: clStr, target }])
-			dispatchContext.logger.info(`_handleSlackBlockAction: unlockchanges IPC result: statusCode=${result?.statusCode}, message='${result?.message}'`)
 			return _slackActionResult(result, payload, `Unlock started by ${who}`, dispatchContext)
 		}
 
@@ -163,7 +160,6 @@ async function _handleCreateShelfSubmit(payload: SlackInteractivePayload, view:
 		?.some((o: any) => o.value === 'acknowledge') ?? false
 
 	const result = await dispatchContext.sendMessage('doNodeOp', [bot, branch, 'create_shelf', { who, cl: String(cl), target, workspace }])
-	dispatchContext.logger.info(`_handleCreateShelfSubmit: create_shelf IPC result: statusCode=${result?.statusCode}, message='${result?.message}'`)
 	if (result && result.statusCode === 200) {
 		if (acknowledgeSelected) {
 			// Mirror the web UI: acknowledge via edge op when an edge is present
@@ -215,7 +211,6 @@ async function _handleSkipReasonSubmit(payload: SlackInteractivePayload, view: N
 	const result = await dispatchContext.sendMessage('doEdgeOp', [bot, branch, edge, 'set_last_cl', {
 		who, cl: String(cl), reason: reasonText, unblock: 'true', isSkip: 'true'
 	}])
-	dispatchContext.logger.info(`_handleSkipReasonSubmit: set_last_cl IPC result: statusCode=${result?.statusCode}, message='${result?.message}'`)
 
 	if (result && result.statusCode === 200) {
 		if (payload.response_url && dispatchContext.botToken) {
@@ -243,7 +238,6 @@ async function _handleStompConfirmSubmit(payload: SlackInteractivePayload, view:
 	const who = payload.user.name
 
 	const result = await dispatchContext.sendMessage('doNodeOp', [bot, branch, 'stompchanges', { who, cl: String(cl), target }])
-	dispatchContext.logger.info(`_handleStompConfirmSubmit: stompchanges IPC result: statusCode=${result?.statusCode}, message='${result?.message}'`)
 
 	if (result && result.statusCode === 200) {
 		return { statusCode: 200, message: JSON.stringify({ response_action: 'clear' }), headers: [['Content-Type', 'application/json']] }
diff --git a/Engine/Extras/RoboMerge/v3/src/robo/slack-socket.ts b/Engine/Extras/RoboMerge/v3/src/robo/slack-socket.ts
--- a/Engine/Extras/RoboMerge/v3/src/robo/slack-socket.ts
+++ b/Engine/Extras/RoboMerge/v3/src/robo/slack-socket.ts
@@ -100,7 +100,6 @@ export class SlackSocket {
 
 		if (msg.type === 'interactive') {
 			const payload = msg.payload as SlackInteractivePayload
-			this.dispatchContext.logger.info(`Slack Socket Mode: interactive event (subtype=${payload?.type ?? 'unknown'}, envelope_id=${msg.envelope_id ?? 'none'})`)
 			if (payload?.type === 'view_submission') {
 				// For view_submission, the response_action must be included in the envelope ACK.
 				// Do NOT ACK immediately — wait for the dispatch result (within Slack's 3-second window).
@@ -132,8 +131,6 @@ export class SlackSocket {
 			return
 		}
 
-		const actionId = payload.actions?.[0]?.action_id ?? 'unknown'
-		this.dispatchContext.logger.info(`Slack Socket Mode: dispatching block_action '${actionId}' (host='${host ?? 'none'}')`)
 		dispatchSlackInteractivePayload(payload, this.dispatchContext).catch(err => {
 			this.dispatchContext.logger.error(`Slack Socket Mode: dispatch error: ${err}`)
 		})
@@ -154,9 +151,6 @@ export class SlackSocket {
 			return
 		}
 
-		const callbackId = payload.view?.callback_id ?? 'unknown'
-		this.dispatchContext.logger.info(`Slack Socket Mode: dispatching view_submission '${callbackId}' (host='${host ?? 'none'}')`)
-
 		let ackPayload: any = {}
 		try {
 			const result = await dispatchSlackInteractivePayload(payload, this.dispatchContext)
